Copying //UE4/Dev-Sequencer to //UE4/Main (Source: //UE4/Dev-Sequencer @ 3617642)

diff --git a/Engine/Plugins/Media/WmfMedia/Source/WmfMedia/Private/Wmf/WmfMediaUtils.h b/Engine/Plugins/Media/WmfMedia/Source/WmfMedia/Private/Wmf/WmfMediaUtils.h
--- a/Engine/Plugins/Media/WmfMedia/Source/WmfMedia/Private/Wmf/WmfMediaUtils.h
+++ b/Engine/Plugins/Media/WmfMedia/Source/WmfMedia/Private/Wmf/WmfMediaUtils.h
@@ -6,10 +6,45 @@
 
 #if WMFMEDIA_SUPPORTED_PLATFORM
 
+#include "Containers/UnrealString.h"
+#include "Templates/SharedPointer.h"
+
 #include "AllowWindowsPlatformTypes.h"
 
+
 namespace WmfMedia
 {
+	/**
+	 * Create an output media type for the given input media type.
+	 *
+	 * @param MajorType The output's major type, i.e. audio or video.
+	 * @param SubType The output's sub type, i.e. PCM or RGB32.
+	 * @param AllowNonStandardCodecs Whether non-standard codecs should be allowed.
+	 * @return The output type, or NULL on error.
+	 */
+	TComPtr<IMFMediaType> CreateOutputType(const GUID& MajorType, const GUID& SubType, bool AllowNonStandardCodecs);
+
+	/**
+	 * Enumerate the available capture devices.
+	 *
+	 * The supported device types are MF_DEVSOURCE_ATTRIBUTE_SOURCE_TYPE_AUDCAP_GUID for audio
+	 * capture devices and MF_DEVSOURCE_ATTRIBUTE_SOURCE_TYPE_VIDCAP_GUID for video capture devices.
+	 *
+	 * @param DeviceType The type of devices to enumerate.
+	 * @param OutDevices Will contain the collection of found device source activators.
+	 */
+	void EnumerateCaptureDevices(GUID DeviceType, TArray<TComPtr<IMFActivate>>& OutDevices);
+
+	/**
+	 * Convert a floating point frame rate to an integer ratio.
+	 *
+	 * @param OutNumerator Will contain the ratio's numerator.
+	 * @param OutDenominator Will contain the ratio's denominator.
+	 * @return true on success, false if the frame rate cannot be represented as a ratio.
+	 * @see RatioToFrameRate
+	 */
+	bool FrameRateToRatio(float FrameRate, int32& OutNumerator, int32& OutDenominator);
+
 	/** 
 	 * Convert a FOURCC code to string.
 	 *
@@ -18,6 +53,27 @@ namespace WmfMedia
 	 */
 	FString FourccToString(unsigned long Fourcc);
 
+	/**
+	 * Get information about the given capture device.
+	 *
+	 * @param Device The device.
+	 * @param OutDisplayName Will contain the device's human readable name.
+	 * @param OutInfo Will contain a device specific information string.
+	 * @param OutSoftwareDevice Will contain a flag indicating whether the device is software.
+	 * @param OutUrl Will contain the device's media URL string.
+	 * @return true on success, false otherwise.
+	 */
+	bool GetCaptureDeviceInfo(IMFActivate& Device, FText& OutDisplayName, FString& OutInfo, bool& OutSoftwareDevice, FString& OutUrl);
+
+	/**
+	 * Get the list of supported media types for the specified major type.
+	 *
+	 * @param MajorType The major type to query, i.e. audio or video.
+	 * @return Media type collection.
+	 * @see IsSupportedMajorType
+	 */
+	const TArray<TComPtr<IMFMediaType>>& GetSupportedMediaTypes(const GUID& MajorType);
+
 	/**
 	 * Convert a Windows GUID to string.
 	 *
@@ -26,6 +82,15 @@ namespace WmfMedia
 	 */
 	FString GuidToString(const GUID& Guid);
 
+	/**
+	 * Check whether the given major media type is supported.
+	 *
+	 * @param Type The major type to check.
+	 * @return true if the type is supported, false otherwise.
+	 * @see GetSupportedMediaTypes
+	 */
+	bool IsSupportedMajorType(const GUID& MajorType);
+
 	/**
 	 * Convert a media major type to string.
 	 *
@@ -35,20 +100,40 @@ namespace WmfMedia
 	FString MajorTypeToString(const GUID& MajorType);
 
 	/**
-	 * Convert an WMF HRESULT code to string.
+	 * Convert a media event to string.
 	 *
-	 * @param Result The result code to convert.
+	 * @param Event The event code to convert.
 	 * @return The corresponding string.
 	 */
-	FString ResultToString(HRESULT Result);
+	FString MediaEventToString(MediaEventType Event);
 
 	/**
-	 * Convert a WMF session event to string.
+	 * Convert an integer ratio to a floating point frame rate.
 	 *
-	 * @param Event The event code to convert.
+	 * @param Numerator The ratio's numerator.
+	 * @param Denominator The ratio's denominator.
+	 * @return The corresponding floating point value.
+	 * @see FrameRateToRatio
+	 */
+	float RatioToFrameRate(int32 Numerator, int32 Denominator);
+
+	/**
+	 * Resolve a media source from an archive or URL.
+	 *
+	 * @param Archive The archive to read media data from (optional).
+	 * @param Url The media source URL.
+	 * @param Precache Whether to precache media into RAM if URL is a local file.
+	 * @return The media source object, or NULL if it couldn't be resolved.
+	 */
+	TComPtr<IMFMediaSource> ResolveMediaSource(TSharedPtr<FArchive, ESPMode::ThreadSafe> Archive, const FString& Url, bool Precache);
+
+	/**
+	 * Convert an WMF HRESULT code to string.
+	 *
+	 * @param Result The result code to convert.
 	 * @return The corresponding string.
 	 */
-	FString SessionEventToString(MediaEventType Event);
+	FString ResultToString(HRESULT Result);
 
 	/**
 	 * Convert a media sub-type to string.
@@ -62,4 +147,4 @@ namespace WmfMedia
 
 #include "HideWindowsPlatformTypes.h"
 
-#endif
+#endif //WMFMEDIA_SUPPORTED_PLATFORM
diff --git a/Engine/Plugins/Media/WmfMedia/Source/WmfMedia/Private/WmfMediaModule.cpp b/Engine/Plugins/Media/WmfMedia/Source/WmfMedia/Private/WmfMediaModule.cpp
--- a/Engine/Plugins/Media/WmfMedia/Source/WmfMedia/Private/WmfMediaModule.cpp
+++ b/Engine/Plugins/Media/WmfMedia/Source/WmfMedia/Private/WmfMediaModule.cpp
@@ -1,11 +1,16 @@
 // Copyright 1998-2017 Epic Games, Inc. All Rights Reserved.
 
-#include "CoreMinimal.h"
 #include "WmfMediaPrivate.h"
-#include "IWmfMediaModule.h"
 
 #if WMFMEDIA_SUPPORTED_PLATFORM
+	#include "IMediaCaptureSupport.h"
+	#include "IMediaModule.h"
+	#include "Modules/ModuleInterface.h"
+	#include "Modules/ModuleManager.h"
+
+	#include "IWmfMediaModule.h"
 	#include "WmfMediaPlayer.h"
+	#include "WmfMediaUtils.h"
 
 	#pragma comment(lib, "mf")
 	#pragma comment(lib, "mfplat")
@@ -17,14 +22,13 @@
 
 DEFINE_LOG_CATEGORY(LogWmfMedia);
 
-#define LOCTEXT_NAMESPACE "FWmfMediaModule"
-
 
 /**
  * Implements the WmfMedia module.
  */
 class FWmfMediaModule
-	: public IWmfMediaModule
+	: public IMediaCaptureSupport
+	, public IWmfMediaModule
 {
 public:
 
@@ -33,20 +37,40 @@ public:
 		: Initialized(false)
 	{ }
 
+public:
+
+	//~ IMediaCaptureDevices interface
+
+	virtual void EnumerateAudioCaptureDevices(TArray<FMediaCaptureDeviceInfo>& OutDeviceInfos) override
+	{
+#if WMFMEDIA_SUPPORTED_PLATFORM
+		EnumerateCaptureDevices(MF_DEVSOURCE_ATTRIBUTE_SOURCE_TYPE_AUDCAP_GUID, OutDeviceInfos);
+#endif
+	}
+
+	virtual void EnumerateVideoCaptureDevices(TArray<FMediaCaptureDeviceInfo>& OutDeviceInfos) override
+	{
+#if WMFMEDIA_SUPPORTED_PLATFORM
+		EnumerateCaptureDevices(MF_DEVSOURCE_ATTRIBUTE_SOURCE_TYPE_VIDCAP_GUID, OutDeviceInfos);
+#endif
+	}
+
 public:
 
 	//~ IWmfMediaModule interface
 
-	virtual TSharedPtr<IMediaPlayer, ESPMode::ThreadSafe> CreatePlayer() override
+	virtual TSharedPtr<IMediaPlayer, ESPMode::ThreadSafe> CreatePlayer(IMediaEventSink& EventSink) override
 	{
 #if WMFMEDIA_SUPPORTED_PLATFORM
-		if (Initialized)
+		if (!Initialized)
 		{
-			return MakeShareable(new FWmfMediaPlayer());
+			return nullptr;
 		}
-#endif
 
+		return MakeShareable(new FWmfMediaPlayer(EventSink));
+#else
 		return nullptr;
+#endif
 	}
 
 public:
@@ -74,6 +98,14 @@ public:
 			return;
 		}
 
+		// register capture device support
+		auto MediaModule = FModuleManager::LoadModulePtr<IMediaModule>("Media");
+
+		if (MediaModule != nullptr)
+		{
+			MediaModule->RegisterCaptureSupport(*this);
+		}
+
 		Initialized = true;
 
 #endif //WMFMEDIA_SUPPORTED_PLATFORM
@@ -87,6 +119,14 @@ public:
 			return;
 		}
 
+		// unregister capture support
+		auto MediaModule = FModuleManager::GetModulePtr<IMediaModule>("Media");
+
+		if (MediaModule != nullptr)
+		{
+			MediaModule->UnregisterCaptureSupport(*this);
+		}
+
 		// shutdown Windows Media Foundation
 		MFShutdown();
 
@@ -97,6 +137,43 @@ public:
 
 protected:
 
+#if WMFMEDIA_SUPPORTED_PLATFORM
+
+	/**
+	 * Enumerate capture devices of the specified type.
+	 *
+	 * @param DeviceType The type of devices to enumerate.
+	 * @param SoftwareOnly Whether to enumerate software devices only.
+	 * @param OutDeviceInfo Will contain information about the devices.
+	 */
+	void EnumerateCaptureDevices(GUID DeviceType, TArray<FMediaCaptureDeviceInfo>& OutDeviceInfos)
+	{
+		TArray<TComPtr<IMFActivate>> Devices;
+		WmfMedia::EnumerateCaptureDevices(DeviceType, Devices);
+
+		for (auto& Device : Devices)
+		{
+			FMediaCaptureDeviceInfo DeviceInfo;
+			bool SoftwareDevice = false;
+
+			if (!WmfMedia::GetCaptureDeviceInfo(*Device, DeviceInfo.DisplayName, DeviceInfo.Info, SoftwareDevice, DeviceInfo.Url))
+			{
+				continue;
+			}
+
+			if (DeviceType == MF_DEVSOURCE_ATTRIBUTE_SOURCE_TYPE_VIDCAP_GUID)
+			{
+				DeviceInfo.Type = SoftwareDevice ? EMediaCaptureDeviceType::VideoSoftware : EMediaCaptureDeviceType::Video;
+			}
+			else
+			{
+				DeviceInfo.Type = EMediaCaptureDeviceType::Audio;
+			}
+
+			OutDeviceInfos.Add(MoveTemp(DeviceInfo));
+		}
+	}
+
 	/**
 	 * Loads all required Windows libraries.
 	 *
@@ -135,6 +212,8 @@ protected:
 		return true;
 	}
 
+#endif //WMFMEDIA_SUPPORTED_PLATFORM
+
 private:
 
 	/** Whether the module has been initialized. */
@@ -143,6 +222,3 @@ private:
 
 
 IMPLEMENT_MODULE(FWmfMediaModule, WmfMedia);
-
-
-#undef LOCTEXT_NAMESPACE
diff --git a/Engine/Plugins/Media/WmfMedia/Source/WmfMedia/Private/WmfMediaPrivate.h b/Engine/Plugins/Media/WmfMedia/Source/WmfMedia/Private/WmfMediaPrivate.h
--- a/Engine/Plugins/Media/WmfMedia/Source/WmfMedia/Private/WmfMediaPrivate.h
+++ b/Engine/Plugins/Media/WmfMedia/Source/WmfMedia/Private/WmfMediaPrivate.h
@@ -2,42 +2,42 @@
 
 #pragma once
 
-#include "Runtime/Core/Public/CoreMinimal.h"
-#include "Runtime/Core/Public/Modules/ModuleManager.h"
+#include "Logging/LogMacros.h"
 
-#if PLATFORM_WINDOWS
-#include "WindowsHWrapper.h"
-#endif
-
-#define WMFMEDIA_SUPPORTED_PLATFORM (PLATFORM_WINDOWS && WINVER >= _WIN32_WINNT_VISTA)
+#define WMFMEDIA_SUPPORTED_PLATFORM (PLATFORM_WINDOWS && (WINVER >= 0x0600 /*Vista*/))
 
 #if WMFMEDIA_SUPPORTED_PLATFORM
-	#if WITH_ENGINE
-		#include "Runtime/HeadMountedDisplay/Public/IHeadMountedDisplayModule.h"
-	#endif
-
-	#include "Runtime/Media/Public/IMediaAudioSink.h"
-	#include "Runtime/Media/Public/IMediaModule.h"
-	#include "Runtime/Media/Public/IMediaOptions.h"
-	#include "Runtime/Media/Public/IMediaOverlaySink.h"
-	#include "Runtime/Media/Public/IMediaTextureSink.h"
-
 	#include "../../WmfMediaFactory/Public/WmfMediaSettings.h"
 
 	#include "WindowsHWrapper.h"
-	#include "Runtime/Core/Public/Windows/AllowWindowsPlatformTypes.h"
+	#include "Windows/AllowWindowsPlatformTypes.h"
+
 	#include <windows.h>
+	
+	#include <D3D9.h>
 	#include <mfapi.h>
 	#include <mferror.h>
 	#include <mfidl.h>
+	#include <mmdeviceapi.h>
+	#include <mmeapi.h>
 	#include <nserror.h>
 	#include <shlwapi.h>
-	#include "Runtime/Core/Public/Windows/COMPointer.h"
-	#include "Runtime/Core/Public/Windows/HideWindowsPlatformTypes.h"
+
+	const GUID FORMAT_VideoInfo = { 0x05589f80, 0xc356, 0x11ce, { 0xbf, 0x01, 0x00, 0xaa, 0x00, 0x55, 0x59, 0x5a } };
+	const GUID FORMAT_VideoInfo2 = { 0xf72a76A0, 0xeb0a, 0x11d0, { 0xac, 0xe4, 0x00, 0x00, 0xc0, 0xcc, 0x16, 0xba } };
+
+	#if (WINVER < _WIN32_WINNT_WIN8)
+		const GUID MF_LOW_LATENCY = { 0x9c27891a, 0xed7a, 0x40e1, { 0x88, 0xe8, 0xb2, 0x27, 0x27, 0xa0, 0x24, 0xee } };
+	#endif
+
+	#include "Windows/COMPointer.h"
+	#include "Windows/HideWindowsPlatformTypes.h"
+
+#elif PLATFORM_WINDOWS && !UE_SERVER
+	#pragma message("Skipping WmfMedia (requires WINVER >= 0x0600, but WINVER is " PREPROCESSOR_TO_STRING(WINVER) ")")
 
 #endif //WMFMEDIA_SUPPORTED_PLATFORM
 
 
 /** Log category for the WmfMedia module. */
 DECLARE_LOG_CATEGORY_EXTERN(LogWmfMedia, Log, All);
-
diff --git a/Engine/Plugins/Media/WmfMedia/Source/WmfMedia/Public/IWmfMediaModule.h b/Engine/Plugins/Media/WmfMedia/Source/WmfMedia/Public/IWmfMediaModule.h
--- a/Engine/Plugins/Media/WmfMedia/Source/WmfMedia/Public/IWmfMediaModule.h
+++ b/Engine/Plugins/Media/WmfMedia/Source/WmfMedia/Public/IWmfMediaModule.h
@@ -2,12 +2,13 @@
 
 #pragma once
 
-#include "CoreTypes.h"
 #include "Templates/SharedPointer.h"
 #include "Modules/ModuleInterface.h"
 
+class IMediaEventSink;
 class IMediaPlayer;
 
+
 /**
  * Interface for the WmfMedia module.
  */
@@ -19,9 +20,10 @@ public:
 	/**
 	 * Creates a Windows Media Foundation based media player.
 	 *
+	 * @param EventSink The object that receives media events from the player.
 	 * @return A new media player, or nullptr if a player couldn't be created.
 	 */
-	virtual TSharedPtr<IMediaPlayer, ESPMode::ThreadSafe> CreatePlayer() = 0;
+	virtual TSharedPtr<IMediaPlayer, ESPMode::ThreadSafe> CreatePlayer(IMediaEventSink& EventSink) = 0;
 
 public:
 
diff --git a/Engine/Plugins/Media/WmfMedia/Source/WmfMedia/WmfMedia.Build.cs b/Engine/Plugins/Media/WmfMedia/Source/WmfMedia/WmfMedia.Build.cs
--- a/Engine/Plugins/Media/WmfMedia/Source/WmfMedia/WmfMedia.Build.cs
+++ b/Engine/Plugins/Media/WmfMedia/Source/WmfMedia/WmfMedia.Build.cs
@@ -6,49 +6,47 @@ namespace UnrealBuildTool.Rules
 	{
 		public WmfMedia(ReadOnlyTargetRules Target) : base(Target)
 		{
-            DynamicallyLoadedModuleNames.AddRange(
-                new string[] {
-                    "Media",
-				}
-            );
+			DynamicallyLoadedModuleNames.AddRange(
+				new string[] {
+					"Media",
+				});
 
 			PrivateDependencyModuleNames.AddRange(
 				new string[] {
 					"Core",
 					"CoreUObject",
+					"MediaUtils",
 					"RenderCore",
 					"WmfMediaFactory",
-				}
-			);
+				});
 
 			PrivateIncludePathModuleNames.AddRange(
 				new string[] {
-                    "Media",
-				}
-			);
+					"Media",
+				});
 
 			PrivateIncludePaths.AddRange(
 				new string[] {
 					"WmfMedia/Private",
-                    "WmfMedia/Private/Player",
+					"WmfMedia/Private/Player",
 					"WmfMedia/Private/Wmf",
-				}
-			);
+				});
 
 			if (Target.bCompileAgainstEngine)
 			{
+				PrivateDependencyModuleNames.Add("Engine");
 				PrivateDependencyModuleNames.Add("HeadMountedDisplay");
 			}
 
 			if ((Target.Platform == UnrealTargetPlatform.Win64) ||
 				(Target.Platform == UnrealTargetPlatform.Win32))
-            {
-                PublicDelayLoadDLLs.Add("mf.dll");
-                PublicDelayLoadDLLs.Add("mfplat.dll");
-                PublicDelayLoadDLLs.Add("mfplay.dll");
-                PublicDelayLoadDLLs.Add("mfuuid.dll");
-                PublicDelayLoadDLLs.Add("shlwapi.dll");
-            }
+			{
+				PublicDelayLoadDLLs.Add("mf.dll");
+				PublicDelayLoadDLLs.Add("mfplat.dll");
+				PublicDelayLoadDLLs.Add("mfplay.dll");
+				PublicDelayLoadDLLs.Add("mfuuid.dll");
+				PublicDelayLoadDLLs.Add("shlwapi.dll");
+			}
 		}
 	}
 }
diff --git a/Engine/Plugins/Media/WmfMedia/Source/WmfMediaEditor/Private/WmfFileMediaSourceFactory.cpp b/Engine/Plugins/Media/WmfMedia/Source/WmfMediaEditor/Private/WmfFileMediaSourceFactory.cpp
--- a/Engine/Plugins/Media/WmfMedia/Source/WmfMediaEditor/Private/WmfFileMediaSourceFactory.cpp
+++ b/Engine/Plugins/Media/WmfMedia/Source/WmfMediaEditor/Private/WmfFileMediaSourceFactory.cpp
@@ -1,8 +1,11 @@
 // Copyright 1998-2017 Epic Games, Inc. All Rights Reserved.
 
 #include "WmfFileMediaSourceFactory.h"
-#include "Misc/Paths.h"
+
+#include "Containers/UnrealString.h"
 #include "FileMediaSource.h"
+#include "Misc/Paths.h"
+#include "UObject/UObjectGlobals.h"
 
 
 /* UWmfFileMediaSourceFactory structors
diff --git a/Engine/Plugins/Media/WmfMedia/Source/WmfMediaFactory/Private/WmfMediaFactoryClasses.cpp b/Engine/Plugins/Media/WmfMedia/Source/WmfMediaFactory/Private/WmfMediaFactoryClasses.cpp
--- a/Engine/Plugins/Media/WmfMedia/Source/WmfMediaFactory/Private/WmfMediaFactoryClasses.cpp
+++ b/Engine/Plugins/Media/WmfMedia/Source/WmfMediaFactory/Private/WmfMediaFactoryClasses.cpp
@@ -1,10 +1,10 @@
 // Copyright 1998-2017 Epic Games, Inc. All Rights Reserved.
 
-#include "CoreMinimal.h"
 #include "WmfMediaSettings.h"
 
 
 UWmfMediaSettings::UWmfMediaSettings()
 	: AllowNonStandardCodecs(false)
+	, LowLatency(false)
 	, NativeAudioOut(false)
 { }
diff --git a/Engine/Plugins/Media/WmfMedia/Source/WmfMediaFactory/Private/WmfMediaFactoryModule.cpp b/Engine/Plugins/Media/WmfMedia/Source/WmfMediaFactory/Private/WmfMediaFactoryModule.cpp
--- a/Engine/Plugins/Media/WmfMedia/Source/WmfMediaFactory/Private/WmfMediaFactoryModule.cpp
+++ b/Engine/Plugins/Media/WmfMedia/Source/WmfMediaFactory/Private/WmfMediaFactoryModule.cpp
@@ -2,17 +2,16 @@
 
 #include "WmfMediaFactoryPrivate.h"
 
-#include "CoreMinimal.h"
+#include "Containers/Array.h"
+#include "Containers/UnrealString.h"
 #include "IMediaModule.h"
-#include "IWmfMediaModule.h"
 #include "IMediaOptions.h"
 #include "IMediaPlayerFactory.h"
+#include "Internationalization/Internationalization.h"
 #include "Misc/Paths.h"
 #include "Modules/ModuleInterface.h"
 #include "Modules/ModuleManager.h"
-#include "UObject/Class.h"
-#include "UObject/WeakObjectPtr.h"
-#include "WmfMediaSettings.h"
+#include "UObject/NameTypes.h"
 
 #if PLATFORM_WINDOWS
 	#include "WindowsHWrapper.h"
@@ -20,9 +19,14 @@
 
 #if WITH_EDITOR
 	#include "ISettingsModule.h"
-	#include "ISettingsSection.h"
+	#include "Templates/SharedPointer.h"
+	#include "UObject/Class.h"
+	#include "UObject/WeakObjectPtr.h"
+	#include "WmfMediaSettings.h"
 #endif
 
+#include "../../WmfMedia/Public/IWmfMediaModule.h"
+
 
 DEFINE_LOG_CATEGORY(LogWmfMediaFactory);
 
@@ -39,14 +43,13 @@ class FWmfMediaFactoryModule
 public:
 
 	/** Default constructor. */
-	FWmfMediaFactoryModule()
-	{ }
+	FWmfMediaFactoryModule() { }
 
 public:
 
-	//~ IMediaPlayerInfo interface
+	//~ IMediaPlayerFactory interface
 
-	virtual bool CanPlayUrl(const FString& Url, const IMediaOptions& Options, TArray<FText>* OutWarnings, TArray<FText>* OutErrors) const override
+	virtual bool CanPlayUrl(const FString& Url, const IMediaOptions* Options, TArray<FText>* OutWarnings, TArray<FText>* OutErrors) const override
 	{
 		FString Scheme;
 		FString Location;
@@ -89,9 +92,9 @@ public:
 		}
 
 		// check options
-		if (OutWarnings != nullptr)
+		if ((OutWarnings != nullptr) && (Options != nullptr))
 		{
-			if (Options.GetMediaOption("PrecacheFile", false) && (Scheme != TEXT("file")))
+			if (Options->GetMediaOption("PrecacheFile", false) && (Scheme != TEXT("file")))
 			{
 				OutWarnings->Add(LOCTEXT("PrecachingNotSupported", "Precaching is supported for local files only"));
 			}
@@ -100,10 +103,10 @@ public:
 		return true;
 	}
 
-	virtual TSharedPtr<IMediaPlayer, ESPMode::ThreadSafe> CreatePlayer() override
+	virtual TSharedPtr<IMediaPlayer, ESPMode::ThreadSafe> CreatePlayer(IMediaEventSink& EventSink) override
 	{
 		auto WmfMediaModule = FModuleManager::LoadModulePtr<IWmfMediaModule>("WmfMedia");
-		return (WmfMediaModule != nullptr) ? WmfMediaModule->CreatePlayer() : nullptr;
+		return (WmfMediaModule != nullptr) ? WmfMediaModule->CreatePlayer(EventSink) : nullptr;
 	}
 
 	virtual FText GetDisplayName() const override
@@ -122,6 +125,18 @@ public:
 		return SupportedPlatforms;
 	}
 
+	virtual bool SupportsFeature(EMediaFeature Feature) const override
+	{
+		return ((Feature == EMediaFeature::AudioSamples) ||
+				(Feature == EMediaFeature::AudioTracks) ||
+				(Feature == EMediaFeature::CaptionTracks) ||
+				(Feature == EMediaFeature::MetadataTracks) ||
+				(Feature == EMediaFeature::OverlaySamples) ||
+				(Feature == EMediaFeature::SubtitleTracks) ||
+				(Feature == EMediaFeature::VideoSamples) ||
+				(Feature == EMediaFeature::VideoTracks));
+	}
+
 public:
 
 	//~ IModuleInterface interface
@@ -143,6 +158,8 @@ public:
 		SupportedFileExtensions.Add(TEXT("mov"));
 		SupportedFileExtensions.Add(TEXT("mp3"));
 		SupportedFileExtensions.Add(TEXT("mp4"));
+		SupportedFileExtensions.Add(TEXT("sami"));
+		SupportedFileExtensions.Add(TEXT("smi"));
 		SupportedFileExtensions.Add(TEXT("wav"));
 		SupportedFileExtensions.Add(TEXT("wma"));
 		SupportedFileExtensions.Add(TEXT("wmv"));
@@ -151,6 +168,7 @@ public:
 		SupportedPlatforms.Add(TEXT("Windows"));
 
 		// supported schemes
+		SupportedUriSchemes.Add(TEXT("audcap"));
 		SupportedUriSchemes.Add(TEXT("file"));
 		SupportedUriSchemes.Add(TEXT("http"));
 		SupportedUriSchemes.Add(TEXT("httpd"));
@@ -159,6 +177,7 @@ public:
 		SupportedUriSchemes.Add(TEXT("rtsp"));
 		SupportedUriSchemes.Add(TEXT("rtspt"));
 		SupportedUriSchemes.Add(TEXT("rtspu"));
+		SupportedUriSchemes.Add(TEXT("vidcap"));
 
 #if WITH_EDITOR
 		// register settings
@@ -166,7 +185,7 @@ public:
 
 		if (SettingsModule != nullptr)
 		{
-			ISettingsSectionPtr SettingsSection = SettingsModule->RegisterSettings("Project", "Plugins", "WmfMedia",
+			SettingsModule->RegisterSettings("Project", "Plugins", "WmfMedia",
 				LOCTEXT("WmfMediaSettingsName", "WMF Media"),
 				LOCTEXT("WmfMediaSettingsDescription", "Configure the WMF Media plug-in."),
 				GetMutableDefault<UWmfMediaSettings>()
diff --git a/Engine/Plugins/Media/WmfMedia/Source/WmfMediaFactory/Private/WmfMediaFactoryPrivate.h b/Engine/Plugins/Media/WmfMedia/Source/WmfMediaFactory/Private/WmfMediaFactoryPrivate.h
--- a/Engine/Plugins/Media/WmfMedia/Source/WmfMediaFactory/Private/WmfMediaFactoryPrivate.h
+++ b/Engine/Plugins/Media/WmfMedia/Source/WmfMediaFactory/Private/WmfMediaFactoryPrivate.h
@@ -2,7 +2,7 @@
 
 #pragma once
 
-#include "CoreMinimal.h"
+#include "Logging/LogMacros.h"
 
 /** Log category for the WmfMediaFactory module. */
 DECLARE_LOG_CATEGORY_EXTERN(LogWmfMediaFactory, Log, All);
diff --git a/Engine/Plugins/Media/WmfMedia/Source/WmfMediaFactory/Public/WmfMediaSettings.h b/Engine/Plugins/Media/WmfMedia/Source/WmfMediaFactory/Public/WmfMediaSettings.h
--- a/Engine/Plugins/Media/WmfMedia/Source/WmfMediaFactory/Public/WmfMediaSettings.h
+++ b/Engine/Plugins/Media/WmfMedia/Source/WmfMediaFactory/Public/WmfMediaSettings.h
@@ -2,12 +2,15 @@
 
 #pragma once
 
-#include "CoreMinimal.h"
-#include "UObject/ObjectMacros.h"
 #include "UObject/Object.h"
+#include "UObject/ObjectMacros.h"
+
 #include "WmfMediaSettings.generated.h"
 
 
+/**
+ * Settings for the WmfMedia plug-in.
+ */
 UCLASS(config=Engine)
 class WMFMEDIAFACTORY_API UWmfMediaSettings
 	: public UObject
@@ -22,7 +25,7 @@ public:
 public:
 
 	/**
-	 * Whether to allow the loading of media that uses non-standard codecs.
+	 * Whether to allow the loading of media that uses non-standard codecs (default = off).
 	 *
 	 * By default, the player will attempt to detect audio and video codecs that
 	 * are not supported by the operating system out of the box, but may require
@@ -32,7 +35,19 @@ public:
 	UPROPERTY(config, EditAnywhere, Category=Media)
 	bool AllowNonStandardCodecs;
 
-	/** Play audio tracks via the operating system's native sound mixer. */
+	/**
+	 * Enable low latency processing in the Windows media pipeline (default = off).
+	 *
+	 * When this setting is enabled, the media data is generated with the lowest
+	 * possible delay. This might be desirable for certain real-time applications,
+	 * but it can negatively affect audio and video quality.
+	 *
+	 * @note This setting is only supported on Windows 8 or newer
+	 */
+	UPROPERTY(config, EditAnywhere, Category=Media)
+	bool LowLatency;
+
+	/** Play audio tracks via the operating system's native sound mixer (default = off). */
 	UPROPERTY(config, EditAnywhere, Category=Debug)
 	bool NativeAudioOut;
 };
diff --git a/Engine/Plugins/Media/WmfMedia/WmfMedia.uplugin b/Engine/Plugins/Media/WmfMedia/WmfMedia.uplugin
--- a/Engine/Plugins/Media/WmfMedia/WmfMedia.uplugin
+++ b/Engine/Plugins/Media/WmfMedia/WmfMedia.uplugin
@@ -27,6 +27,11 @@
 			"Type" : "Editor",
 			"LoadingPhase" : "PostEngineInit"
 		},
+		{
+			"Name" : "WmfMediaFactory",
+			"Type" : "Editor",
+			"LoadingPhase" : "PostEngineInit"
+		},
 		{
 			"Name" : "WmfMediaFactory",
 			"Type" : "RuntimeNoCommandlet",
diff --git a/Engine/Plugins/Messaging/MessagingDebugger/Source/MessagingDebugger/Private/MessagingDebuggerModule.cpp b/Engine/Plugins/Messaging/MessagingDebugger/Source/MessagingDebugger/Private/MessagingDebuggerModule.cpp
--- a/Engine/Plugins/Messaging/MessagingDebugger/Source/MessagingDebugger/Private/MessagingDebuggerModule.cpp
+++ b/Engine/Plugins/Messaging/MessagingDebugger/Source/MessagingDebugger/Private/MessagingDebuggerModule.cpp
@@ -1,19 +1,22 @@
 // Copyright 1998-2017 Epic Games, Inc. All Rights Reserved.
 
-#include "CoreMinimal.h"
+#include "CoreTypes.h"
 #include "Features/IModularFeature.h"
-#include "Modules/ModuleInterface.h"
-#include "Modules/ModuleManager.h"
-#include "IMessagingModule.h"
-#include "Widgets/DeclarativeSyntaxSupport.h"
-#include "Textures/SlateIcon.h"
-#include "Widgets/Text/STextBlock.h"
+#include "Features/IModularFeatures.h"
 #include "Framework/Docking/TabManager.h"
+#include "IMessageBus.h"
+#include "IMessagingModule.h"
 #include "Models/MessagingDebuggerCommands.h"
+#include "Modules/ModuleInterface.h"
+#include "Modules/ModuleManager.h"
 #include "Styles/MessagingDebuggerStyle.h"
-#include "Widgets/SMessagingDebugger.h"
-#include "Features/IModularFeatures.h"
+#include "Templates/SharedPointer.h"
+#include "Textures/SlateIcon.h"
+#include "Widgets/DeclarativeSyntaxSupport.h"
 #include "Widgets/Docking/SDockTab.h"
+#include "Widgets/SMessagingDebugger.h"
+#include "Widgets/Text/STextBlock.h"
+#include "UObject/NameTypes.h"
 
 #if WITH_EDITOR
 	#include "WorkspaceMenuStructure.h"
diff --git a/Engine/Plugins/Messaging/MessagingDebugger/Source/MessagingDebugger/Private/Models/MessagingDebuggerCommands.h b/Engine/Plugins/Messaging/MessagingDebugger/Source/MessagingDebugger/Private/Models/MessagingDebuggerCommands.h
--- a/Engine/Plugins/Messaging/MessagingDebugger/Source/MessagingDebugger/Private/Models/MessagingDebuggerCommands.h
+++ b/Engine/Plugins/Messaging/MessagingDebugger/Source/MessagingDebugger/Private/Models/MessagingDebuggerCommands.h
@@ -2,13 +2,16 @@
 
 #pragma once
 
-#include "CoreMinimal.h"
+#include "CoreTypes.h"
 #include "InputCoreTypes.h"
 #include "Framework/Commands/InputChord.h"
 #include "Framework/Commands/Commands.h"
+#include "Templates/SharedPointer.h"
+
 
 #define LOCTEXT_NAMESPACE "FMessagingDebuggerCommands"
 
+
 class FMessagingDebuggerCommands
 	: public TCommands<FMessagingDebuggerCommands>
 {
@@ -21,7 +24,7 @@ public:
 
 public:
 
-	// TCommands interface
+	//~ TCommands interface
 
 	virtual void RegisterCommands() override
 	{
diff --git a/Engine/Plugins/Messaging/MessagingDebugger/Source/MessagingDebugger/Private/Models/MessagingDebuggerEndpointFilter.h b/Engine/Plugins/Messaging/MessagingDebugger/Source/MessagingDebugger/Private/Models/MessagingDebuggerEndpointFilter.h
--- a/Engine/Plugins/Messaging/MessagingDebugger/Source/MessagingDebugger/Private/Models/MessagingDebuggerEndpointFilter.h
+++ b/Engine/Plugins/Messaging/MessagingDebugger/Source/MessagingDebugger/Private/Models/MessagingDebuggerEndpointFilter.h
@@ -2,8 +2,13 @@
 
 #pragma once
 
-#include "CoreMinimal.h"
+#include "CoreTypes.h"
+#include "Containers/Array.h"
+#include "Containers/UnrealString.h"
+#include "Delegates/Delegate.h"
 #include "IMessageTracer.h"
+#include "Templates/SharedPointer.h"
+
 
 /**
  * Implements a filter for the message endpoints list.
diff --git a/Engine/Plugins/Messaging/MessagingDebugger/Source/MessagingDebugger/Private/Models/MessagingDebuggerMessageFilter.h b/Engine/Plugins/Messaging/MessagingDebugger/Source/MessagingDebugger/Private/Models/MessagingDebuggerMessageFilter.h
--- a/Engine/Plugins/Messaging/MessagingDebugger/Source/MessagingDebugger/Private/Models/MessagingDebuggerMessageFilter.h
+++ b/Engine/Plugins/Messaging/MessagingDebugger/Source/MessagingDebugger/Private/Models/MessagingDebuggerMessageFilter.h
@@ -2,8 +2,10 @@
 
 #pragma once
 
-#include "CoreMinimal.h"
+#include "CoreTypes.h"
 #include "IMessageTracer.h"
+#include "Templates/SharedPointer.h"
+
 
 /**
  * Implements a filter for the message history list.
diff --git a/Engine/Plugins/Messaging/MessagingDebugger/Source/MessagingDebugger/Private/Models/MessagingDebuggerModel.cpp b/Engine/Plugins/Messaging/MessagingDebugger/Source/MessagingDebugger/Private/Models/MessagingDebuggerModel.cpp
--- a/Engine/Plugins/Messaging/MessagingDebugger/Source/MessagingDebugger/Private/Models/MessagingDebuggerModel.cpp
+++ b/Engine/Plugins/Messaging/MessagingDebugger/Source/MessagingDebugger/Private/Models/MessagingDebuggerModel.cpp
@@ -2,6 +2,8 @@
 
 #include "Models/MessagingDebuggerModel.h"
 
+#include "IMessageTracer.h"
+
 
 /* FMessagingDebuggerModel interface
  *****************************************************************************/
diff --git a/Engine/Plugins/Messaging/MessagingDebugger/Source/MessagingDebugger/Private/Models/MessagingDebuggerModel.h b/Engine/Plugins/Messaging/MessagingDebugger/Source/MessagingDebugger/Private/Models/MessagingDebuggerModel.h
--- a/Engine/Plugins/Messaging/MessagingDebugger/Source/MessagingDebugger/Private/Models/MessagingDebuggerModel.h
+++ b/Engine/Plugins/Messaging/MessagingDebugger/Source/MessagingDebugger/Private/Models/MessagingDebuggerModel.h
@@ -2,8 +2,16 @@
 
 #pragma once
 
-#include "CoreMinimal.h"
+#include "CoreTypes.h"
+#include "Containers/Array.h"
+#include "Delegates/Delegate.h"
 #include "IMessageTracer.h"
+#include "Templates/SharedPointer.h"
+
+struct FMessageTracerEndpointInfo;
+struct FMessageTracerMessageInfo;
+struct FMessageTracerTypeInfo;
+
 
 /**
  * Implements a view model for the messaging debugger.
diff --git a/Engine/Plugins/Messaging/MessagingDebugger/Source/MessagingDebugger/Private/Models/MessagingDebuggerTypeFilter.h b/Engine/Plugins/Messaging/MessagingDebugger/Source/MessagingDebugger/Private/Models/MessagingDebuggerTypeFilter.h
--- a/Engine/Plugins/Messaging/MessagingDebugger/Source/MessagingDebugger/Private/Models/MessagingDebuggerTypeFilter.h
+++ b/Engine/Plugins/Messaging/MessagingDebugger/Source/MessagingDebugger/Private/Models/MessagingDebuggerTypeFilter.h
@@ -2,8 +2,13 @@
 
 #pragma once
 
-#include "CoreMinimal.h"
+#include "CoreTypes.h"
+#include "Containers/Array.h"
+#include "Containers/UnrealString.h"
+#include "Delegates/Delegate.h"
 #include "IMessageTracer.h"
+#include "Templates/SharedPointer.h"
+
 
 /**
  * Implements a filter for the message endpoints list.
diff --git a/Engine/Plugins/Messaging/MessagingDebugger/Source/MessagingDebugger/Private/Styles/MessagingDebuggerStyle.h b/Engine/Plugins/Messaging/MessagingDebugger/Source/MessagingDebugger/Private/Styles/MessagingDebuggerStyle.h
--- a/Engine/Plugins/Messaging/MessagingDebugger/Source/MessagingDebugger/Private/Styles/MessagingDebuggerStyle.h
+++ b/Engine/Plugins/Messaging/MessagingDebugger/Source/MessagingDebugger/Private/Styles/MessagingDebuggerStyle.h
@@ -3,15 +3,15 @@
 #pragma once
 
 #include "CoreMinimal.h"
-#include "Layout/Margin.h"
-#include "Fonts/SlateFontInfo.h"
-#include "Styling/SlateTypes.h"
-#include "Misc/Paths.h"
 #include "Brushes/SlateBorderBrush.h"
 #include "Brushes/SlateBoxBrush.h"
 #include "Brushes/SlateImageBrush.h"
+#include "Fonts/SlateFontInfo.h"
+#include "Layout/Margin.h"
+#include "Misc/Paths.h"
 #include "Styling/SlateStyleRegistry.h"
 #include "Styling/SlateStyle.h"
+#include "Styling/SlateTypes.h"
 
 
 #define IMAGE_BRUSH(RelativePath, ...) FSlateImageBrush(RootToContentDir(RelativePath, TEXT(".png")), __VA_ARGS__)
@@ -106,8 +106,8 @@ public:
 		FSlateStyleRegistry::RegisterSlateStyle(*this);
 	 }
 
-	 /** Destructor. */
-	 ~FMessagingDebuggerStyle()
+	 /** Virtual destructor. */
+	 virtual ~FMessagingDebuggerStyle()
 	 {
 		FSlateStyleRegistry::UnRegisterSlateStyle(*this);
 	 }
diff --git a/Engine/Plugins/Messaging/MessagingDebugger/Source/MessagingDebugger/Private/Widgets/Breakpoints/SMessagingBreakpoints.cpp b/Engine/Plugins/Messaging/MessagingDebugger/Source/MessagingDebugger/Private/Widgets/Breakpoints/SMessagingBreakpoints.cpp
--- a/Engine/Plugins/Messaging/MessagingDebugger/Source/MessagingDebugger/Private/Widgets/Breakpoints/SMessagingBreakpoints.cpp
+++ b/Engine/Plugins/Messaging/MessagingDebugger/Source/MessagingDebugger/Private/Widgets/Breakpoints/SMessagingBreakpoints.cpp
@@ -1,7 +1,13 @@
 // Copyright 1998-2017 Epic Games, Inc. All Rights Reserved.
 
 #include "Widgets/Breakpoints/SMessagingBreakpoints.h"
-#include "Widgets/Views/SListView.h"
+
+#include "IMessageTracer.h"
+#include "IMessageTracerBreakpoint.h"
+#include "Styling/ISlateStyle.h"
+#include "Widgets/Views/STableViewBase.h"
+#include "Widgets/Views/STableRow.h"
+
 #include "Widgets/Breakpoints/SMessagingBreakpointsTableRow.h"
 
 
diff --git a/Engine/Plugins/Messaging/MessagingDebugger/Source/MessagingDebugger/Private/Widgets/Breakpoints/SMessagingBreakpoints.h b/Engine/Plugins/Messaging/MessagingDebugger/Source/MessagingDebugger/Private/Widgets/Breakpoints/SMessagingBreakpoints.h
--- a/Engine/Plugins/Messaging/MessagingDebugger/Source/MessagingDebugger/Private/Widgets/Breakpoints/SMessagingBreakpoints.h
+++ b/Engine/Plugins/Messaging/MessagingDebugger/Source/MessagingDebugger/Private/Widgets/Breakpoints/SMessagingBreakpoints.h
@@ -2,14 +2,18 @@
 
 #pragma once
 
-#include "CoreMinimal.h"
-#include "IMessageTracerBreakpoint.h"
-#include "IMessageTracer.h"
-#include "Styling/ISlateStyle.h"
+#include "CoreTypes.h"
+#include "Containers/Array.h"
+#include "Templates/SharedPointer.h"
 #include "Widgets/DeclarativeSyntaxSupport.h"
 #include "Widgets/SCompoundWidget.h"
-#include "Widgets/Views/STableViewBase.h"
-#include "Widgets/Views/STableRow.h"
+#include "Widgets/Views/SListView.h"
+
+class IMessageTracer;
+class IMessageTracerBreakpoint;
+class ISlateStyle;
+class STableViewBase;
+
 
 /**
  * Implements the message breakpoints panel.
@@ -25,7 +29,7 @@ public:
 public:
 
 	/**
-	 * Construct this widget
+	 * Construct this widget.
 	 *
 	 * @param InArgs The declaration data for this widget.
 	 * @param InStyle The visual style to use for this widget.
diff --git a/Engine/Plugins/Messaging/MessagingDebugger/Source/MessagingDebugger/Private/Widgets/Breakpoints/SMessagingBreakpointsTableRow.cpp b/Engine/Plugins/Messaging/MessagingDebugger/Source/MessagingDebugger/Private/Widgets/Breakpoints/SMessagingBreakpointsTableRow.cpp
--- a/Engine/Plugins/Messaging/MessagingDebugger/Source/MessagingDebugger/Private/Widgets/Breakpoints/SMessagingBreakpointsTableRow.cpp
+++ b/Engine/Plugins/Messaging/MessagingDebugger/Source/MessagingDebugger/Private/Widgets/Breakpoints/SMessagingBreakpointsTableRow.cpp
@@ -2,6 +2,13 @@
 
 #include "Widgets/Breakpoints/SMessagingBreakpointsTableRow.h"
 
+#include "IMessageTracerBreakpoint.h"
+#include "SlateOptMacros.h"
+#include "Styling/ISlateStyle.h"
+#include "Widgets/Images/SImage.h"
+#include "Widgets/Input/SCheckBox.h"
+#include "Widgets/Text/STextBlock.h"
+
 
 #define LOCTEXT_NAMESPACE "SMessagingBreakpointsTableRow"
 
diff --git a/Engine/Plugins/Messaging/MessagingDebugger/Source/MessagingDebugger/Private/Widgets/Breakpoints/SMessagingBreakpointsTableRow.h b/Engine/Plugins/Messaging/MessagingDebugger/Source/MessagingDebugger/Private/Widgets/Breakpoints/SMessagingBreakpointsTableRow.h
--- a/Engine/Plugins/Messaging/MessagingDebugger/Source/MessagingDebugger/Private/Widgets/Breakpoints/SMessagingBreakpointsTableRow.h
+++ b/Engine/Plugins/Messaging/MessagingDebugger/Source/MessagingDebugger/Private/Widgets/Breakpoints/SMessagingBreakpointsTableRow.h
@@ -2,16 +2,16 @@
 
 #pragma once
 
-#include "CoreMinimal.h"
-#include "IMessageTracerBreakpoint.h"
+#include "CoreTypes.h"
+#include "Templates/SharedPointer.h"
 #include "Widgets/DeclarativeSyntaxSupport.h"
-#include "Widgets/SWidget.h"
-#include "Widgets/Views/STableViewBase.h"
-#include "Widgets/Views/STableRow.h"
 #include "Widgets/Views/SListView.h"
-#include "SlateOptMacros.h"
-#include "Widgets/Images/SImage.h"
-#include "Widgets/Input/SCheckBox.h"
+#include "Widgets/Views/STableRow.h"
+
+class IMessageTracerBreakpoint;
+class ISlateStyle;
+class STableViewBase;
+
 
 /**
  * Implements a row widget for the session console log.
diff --git a/Engine/Plugins/Messaging/MessagingDebugger/Source/MessagingDebugger/Private/Widgets/EndpointDetails/SMessagingAddressTableRow.cpp b/Engine/Plugins/Messaging/MessagingDebugger/Source/MessagingDebugger/Private/Widgets/EndpointDetails/SMessagingAddressTableRow.cpp
--- a/Engine/Plugins/Messaging/MessagingDebugger/Source/MessagingDebugger/Private/Widgets/EndpointDetails/SMessagingAddressTableRow.cpp
+++ b/Engine/Plugins/Messaging/MessagingDebugger/Source/MessagingDebugger/Private/Widgets/EndpointDetails/SMessagingAddressTableRow.cpp
@@ -2,6 +2,11 @@
 
 #include "Widgets/EndpointDetails/SMessagingAddressTableRow.h"
 
+#include "IMessageTracer.h"
+#include "SlateOptMacros.h"
+
+#include "Models/MessagingDebuggerModel.h"
+
 
 namespace MessagingAddressTableRow
 {
diff --git a/Engine/Plugins/Messaging/MessagingDebugger/Source/MessagingDebugger/Private/Widgets/EndpointDetails/SMessagingAddressTableRow.h b/Engine/Plugins/Messaging/MessagingDebugger/Source/MessagingDebugger/Private/Widgets/EndpointDetails/SMessagingAddressTableRow.h
--- a/Engine/Plugins/Messaging/MessagingDebugger/Source/MessagingDebugger/Private/Widgets/EndpointDetails/SMessagingAddressTableRow.h
+++ b/Engine/Plugins/Messaging/MessagingDebugger/Source/MessagingDebugger/Private/Widgets/EndpointDetails/SMessagingAddressTableRow.h
@@ -2,14 +2,16 @@
 
 #pragma once
 
-#include "CoreMinimal.h"
+#include "CoreTypes.h"
 #include "Widgets/DeclarativeSyntaxSupport.h"
-#include "Widgets/SWidget.h"
 #include "Widgets/Views/STableViewBase.h"
-#include "Widgets/Views/STableRow.h"
-#include "Models/MessagingDebuggerModel.h"
 #include "Widgets/Views/SListView.h"
-#include "SlateOptMacros.h"
+
+class FMessagingDebuggerModel;
+class ISlateStyle;
+
+struct FMessageTracerAddressInfo;
+
 
 /**
  * Implements a row widget for the dispatch state list.
diff --git a/Engine/Plugins/Messaging/MessagingDebugger/Source/MessagingDebugger/Private/Widgets/EndpointDetails/SMessagingEndpointDetails.cpp b/Engine/Plugins/Messaging/MessagingDebugger/Source/MessagingDebugger/Private/Widgets/EndpointDetails/SMessagingEndpointDetails.cpp
--- a/Engine/Plugins/Messaging/MessagingDebugger/Source/MessagingDebugger/Private/Widgets/EndpointDetails/SMessagingEndpointDetails.cpp
+++ b/Engine/Plugins/Messaging/MessagingDebugger/Source/MessagingDebugger/Private/Widgets/EndpointDetails/SMessagingEndpointDetails.cpp
@@ -1,9 +1,14 @@
 // Copyright 1998-2017 Epic Games, Inc. All Rights Reserved.
 
 #include "Widgets/EndpointDetails/SMessagingEndpointDetails.h"
-#include "Widgets/Text/STextBlock.h"
+
+#include "IMessageTracer.h"
+#include "Styling/ISlateStyle.h"
 #include "Widgets/Layout/SGridPanel.h"
-#include "Widgets/Views/SListView.h"
+#include "Widgets/Text/STextBlock.h"
+#include "Widgets/Views/STableViewBase.h"
+
+#include "Models/MessagingDebuggerModel.h"
 #include "Widgets/EndpointDetails/SMessagingAddressTableRow.h"
 
 
diff --git a/Engine/Plugins/Messaging/MessagingDebugger/Source/MessagingDebugger/Private/Widgets/EndpointDetails/SMessagingEndpointDetails.h b/Engine/Plugins/Messaging/MessagingDebugger/Source/MessagingDebugger/Private/Widgets/EndpointDetails/SMessagingEndpointDetails.h
--- a/Engine/Plugins/Messaging/MessagingDebugger/Source/MessagingDebugger/Private/Widgets/EndpointDetails/SMessagingEndpointDetails.h
+++ b/Engine/Plugins/Messaging/MessagingDebugger/Source/MessagingDebugger/Private/Widgets/EndpointDetails/SMessagingEndpointDetails.h
@@ -2,13 +2,19 @@
 
 #pragma once
 
-#include "CoreMinimal.h"
-#include "Styling/ISlateStyle.h"
+#include "CoreTypes.h"
+#include "Containers/Array.h"
+#include "Templates/SharedPointer.h"
 #include "Widgets/DeclarativeSyntaxSupport.h"
 #include "Widgets/SCompoundWidget.h"
-#include "Widgets/Views/STableViewBase.h"
-#include "Widgets/Views/STableRow.h"
-#include "Models/MessagingDebuggerModel.h"
+#include "Widgets/Views/SListView.h"
+
+class FMessagingDebuggerModel;
+class ISlateStyle;
+class STableViewBase;
+
+struct FMessageTracerAddressInfo;
+
 
 /**
 * Implements the message types panel.
@@ -24,7 +30,7 @@ public:
 public:
 
 	/**
-	 * Construct this widget
+	 * Construct this widget.
 	 *
 	 * @param InArgs The declaration data for this widget.
 	 * @param InModel The view model to use.
diff --git a/Engine/Plugins/Messaging/MessagingDebugger/Source/MessagingDebugger/Private/Widgets/Endpoints/SMessagingEndpoints.cpp b/Engine/Plugins/Messaging/MessagingDebugger/Source/MessagingDebugger/Private/Widgets/Endpoints/SMessagingEndpoints.cpp
--- a/Engine/Plugins/Messaging/MessagingDebugger/Source/MessagingDebugger/Private/Widgets/Endpoints/SMessagingEndpoints.cpp
+++ b/Engine/Plugins/Messaging/MessagingDebugger/Source/MessagingDebugger/Private/Widgets/Endpoints/SMessagingEndpoints.cpp
@@ -1,12 +1,18 @@
 // Copyright 1998-2017 Epic Games, Inc. All Rights Reserved.
 
 #include "Widgets/Endpoints/SMessagingEndpoints.h"
+
+#include "IMessageTracer.h"
 #include "SlateOptMacros.h"
+#include "Styling/ISlateStyle.h"
 #include "Widgets/Images/SImage.h"
+#include "Widgets/Layout/SExpandableArea.h"
+#include "Widgets/Views/STableViewBase.h"
+
+#include "Models/MessagingDebuggerEndpointFilter.h"
+#include "Models/MessagingDebuggerModel.h"
 #include "Widgets/Endpoints/SMessagingEndpointsFilterBar.h"
-#include "Widgets/Views/SListView.h"
 #include "Widgets/Endpoints/SMessagingEndpointsTableRow.h"
-#include "Widgets/Layout/SExpandableArea.h"
 
 
 #define LOCTEXT_NAMESPACE "SMessagingEndpoints"
diff --git a/Engine/Plugins/Messaging/MessagingDebugger/Source/MessagingDebugger/Private/Widgets/Endpoints/SMessagingEndpoints.h b/Engine/Plugins/Messaging/MessagingDebugger/Source/MessagingDebugger/Private/Widgets/Endpoints/SMessagingEndpoints.h
--- a/Engine/Plugins/Messaging/MessagingDebugger/Source/MessagingDebugger/Private/Widgets/Endpoints/SMessagingEndpoints.h
+++ b/Engine/Plugins/Messaging/MessagingDebugger/Source/MessagingDebugger/Private/Widgets/Endpoints/SMessagingEndpoints.h
@@ -2,15 +2,21 @@
 
 #pragma once
 
-#include "CoreMinimal.h"
-#include "Styling/ISlateStyle.h"
+#include "CoreTypes.h"
+#include "Containers/Array.h"
+#include "Templates/SharedPointer.h"
 #include "Widgets/DeclarativeSyntaxSupport.h"
 #include "Widgets/SCompoundWidget.h"
-#include "Widgets/Views/STableViewBase.h"
-#include "Widgets/Views/STableRow.h"
-#include "IMessageTracer.h"
-#include "Models/MessagingDebuggerEndpointFilter.h"
-#include "Models/MessagingDebuggerModel.h"
+#include "Widgets/Views/SListView.h"
+
+class FMessagingDebuggerEndpointFilter;
+class FMessagingDebuggerModel;
+class IMessageTracer;
+class ISlateStyle;
+class STableViewBase;
+
+struct FMessageTracerEndpointInfo;
+
 
 /**
  * Implements the message endpoints panel.
@@ -31,7 +37,7 @@ public:
 public:
 
 	/**
-	 * Construct this widget
+	 * Construct this widget.
 	 *
 	 * @param InArgs The declaration data for this widget.
 	 * @param InModel The view model to use.
diff --git a/Engine/Plugins/Messaging/MessagingDebugger/Source/MessagingDebugger/Private/Widgets/Endpoints/SMessagingEndpointsFilterBar.cpp b/Engine/Plugins/Messaging/MessagingDebugger/Source/MessagingDebugger/Private/Widgets/Endpoints/SMessagingEndpointsFilterBar.cpp
--- a/Engine/Plugins/Messaging/MessagingDebugger/Source/MessagingDebugger/Private/Widgets/Endpoints/SMessagingEndpointsFilterBar.cpp
+++ b/Engine/Plugins/Messaging/MessagingDebugger/Source/MessagingDebugger/Private/Widgets/Endpoints/SMessagingEndpointsFilterBar.cpp
@@ -1,9 +1,12 @@
 // Copyright 1998-2017 Epic Games, Inc. All Rights Reserved.
 
 #include "Widgets/Endpoints/SMessagingEndpointsFilterBar.h"
+
 #include "Widgets/SBoxPanel.h"
 #include "Widgets/Input/SSearchBox.h"
 
+#include "Models/MessagingDebuggerEndpointFilter.h"
+
 
 #define LOCTEXT_NAMESPACE "SMessagingEndpointsFilterBar"
 
diff --git a/Engine/Plugins/Messaging/MessagingDebugger/Source/MessagingDebugger/Private/Widgets/Endpoints/SMessagingEndpointsFilterBar.h b/Engine/Plugins/Messaging/MessagingDebugger/Source/MessagingDebugger/Private/Widgets/Endpoints/SMessagingEndpointsFilterBar.h
--- a/Engine/Plugins/Messaging/MessagingDebugger/Source/MessagingDebugger/Private/Widgets/Endpoints/SMessagingEndpointsFilterBar.h
+++ b/Engine/Plugins/Messaging/MessagingDebugger/Source/MessagingDebugger/Private/Widgets/Endpoints/SMessagingEndpointsFilterBar.h
@@ -2,10 +2,13 @@
 
 #pragma once
 
-#include "CoreMinimal.h"
+#include "CoreTypes.h"
+#include "Templates/SharedPointer.h"
 #include "Widgets/DeclarativeSyntaxSupport.h"
 #include "Widgets/SCompoundWidget.h"
-#include "Models/MessagingDebuggerEndpointFilter.h"
+
+class FMessagingDebuggerEndpointFilter;
+
 
 /**
  * Implements the endpoints list filter bar widget.
@@ -21,7 +24,7 @@ public:
 public:
 
 	/**
-	 * Construct this widget
+	 * Construct this widget.
 	 *
 	 * @param InArgs The declaration data for this widget.
 	 * @param InFilter The filter model.
diff --git a/Engine/Plugins/Messaging/MessagingDebugger/Source/MessagingDebugger/Private/Widgets/Endpoints/SMessagingEndpointsTableRow.cpp b/Engine/Plugins/Messaging/MessagingDebugger/Source/MessagingDebugger/Private/Widgets/Endpoints/SMessagingEndpointsTableRow.cpp
--- a/Engine/Plugins/Messaging/MessagingDebugger/Source/MessagingDebugger/Private/Widgets/Endpoints/SMessagingEndpointsTableRow.cpp
+++ b/Engine/Plugins/Messaging/MessagingDebugger/Source/MessagingDebugger/Private/Widgets/Endpoints/SMessagingEndpointsTableRow.cpp
@@ -2,6 +2,16 @@
 
 #include "Widgets/Endpoints/SMessagingEndpointsTableRow.h"
 
+#include "IMessageTracer.h"
+#include "SlateOptMacros.h"
+#include "Widgets/SWidget.h"
+#include "Widgets/Images/SImage.h"
+#include "Widgets/Input/SCheckBox.h"
+#include "Widgets/Text/STextBlock.h"
+#include "Widgets/Views/STableViewBase.h"
+
+#include "Models/MessagingDebuggerModel.h"
+
 
 #define LOCTEXT_NAMESPACE "SMessagingEndpointsTableRow"
 
diff --git a/Engine/Plugins/Messaging/MessagingDebugger/Source/MessagingDebugger/Private/Widgets/Endpoints/SMessagingEndpointsTableRow.h b/Engine/Plugins/Messaging/MessagingDebugger/Source/MessagingDebugger/Private/Widgets/Endpoints/SMessagingEndpointsTableRow.h
--- a/Engine/Plugins/Messaging/MessagingDebugger/Source/MessagingDebugger/Private/Widgets/Endpoints/SMessagingEndpointsTableRow.h
+++ b/Engine/Plugins/Messaging/MessagingDebugger/Source/MessagingDebugger/Private/Widgets/Endpoints/SMessagingEndpointsTableRow.h
@@ -3,16 +3,20 @@
 #pragma once
 
 #include "CoreMinimal.h"
+#include "Internationalization/Text.h"
 #include "Misc/Attribute.h"
+#include "Templates/SharedPointer.h"
 #include "Widgets/DeclarativeSyntaxSupport.h"
-#include "Widgets/SWidget.h"
-#include "Widgets/Views/STableViewBase.h"
-#include "Widgets/Views/STableRow.h"
-#include "Models/MessagingDebuggerModel.h"
-#include "SlateOptMacros.h"
-#include "Widgets/Images/SImage.h"
 #include "Widgets/Views/SListView.h"
-#include "Widgets/Input/SCheckBox.h"
+#include "Widgets/Views/STableRow.h"
+
+class FMessagingDebuggerModel;
+class ISlateStyle;
+class STableViewBase;
+class SWidget;
+
+struct FMessageTracerEndpointInfo;
+
 
 /**
  * Implements a row widget for the session console log.
diff --git a/Engine/Plugins/Messaging/MessagingDebugger/Source/MessagingDebugger/Private/Widgets/Graph/SMessagingGraph.cpp b/Engine/Plugins/Messaging/MessagingDebugger/Source/MessagingDebugger/Private/Widgets/Graph/SMessagingGraph.cpp
--- a/Engine/Plugins/Messaging/MessagingDebugger/Source/MessagingDebugger/Private/Widgets/Graph/SMessagingGraph.cpp
+++ b/Engine/Plugins/Messaging/MessagingDebugger/Source/MessagingDebugger/Private/Widgets/Graph/SMessagingGraph.cpp
@@ -2,6 +2,8 @@
 
 #include "Widgets/Graph/SMessagingGraph.h"
 
+#include "Styling/ISlateStyle.h"
+
 
 #define LOCTEXT_NAMESPACE "SMessagingGraph"
 
diff --git a/Engine/Plugins/Messaging/MessagingDebugger/Source/MessagingDebugger/Private/Widgets/Graph/SMessagingGraph.h b/Engine/Plugins/Messaging/MessagingDebugger/Source/MessagingDebugger/Private/Widgets/Graph/SMessagingGraph.h
--- a/Engine/Plugins/Messaging/MessagingDebugger/Source/MessagingDebugger/Private/Widgets/Graph/SMessagingGraph.h
+++ b/Engine/Plugins/Messaging/MessagingDebugger/Source/MessagingDebugger/Private/Widgets/Graph/SMessagingGraph.h
@@ -2,11 +2,13 @@
 
 #pragma once
 
-#include "CoreMinimal.h"
-#include "Styling/ISlateStyle.h"
+#include "CoreTypes.h"
 #include "Widgets/DeclarativeSyntaxSupport.h"
 #include "Widgets/SCompoundWidget.h"
 
+class ISlateStyle;
+
+
 /**
  * Implements the message interaction graph panel.
  */
@@ -21,12 +23,12 @@ public:
 public:
 
 	/**
-	 * Construct this widget
+	 * Construct this widget.
 	 *
 	 * @param InArgs The declaration data for this widget.
 	 * @param InStyle The visual style to use for this widget.
 	 */
-	void Construct( const FArguments& InArgs, const TSharedRef<ISlateStyle>& InStyle );
+	void Construct(const FArguments& InArgs, const TSharedRef<ISlateStyle>& InStyle);
 
 private:
 
diff --git a/Engine/Plugins/Messaging/MessagingDebugger/Source/MessagingDebugger/Private/Widgets/History/SMessagingHistory.cpp b/Engine/Plugins/Messaging/MessagingDebugger/Source/MessagingDebugger/Private/Widgets/History/SMessagingHistory.cpp
--- a/Engine/Plugins/Messaging/MessagingDebugger/Source/MessagingDebugger/Private/Widgets/History/SMessagingHistory.cpp
+++ b/Engine/Plugins/Messaging/MessagingDebugger/Source/MessagingDebugger/Private/Widgets/History/SMessagingHistory.cpp
@@ -1,12 +1,19 @@
 // Copyright 1998-2017 Epic Games, Inc. All Rights Reserved.
 
 #include "Widgets/History/SMessagingHistory.h"
+
+#include "IMessageTracer.h"
 #include "SlateOptMacros.h"
+#include "Styling/ISlateStyle.h"
+#include "Widgets/Input/SHyperlink.h"
 #include "Widgets/Text/STextBlock.h"
+#include "Widgets/Views/STableViewBase.h"
+#include "Widgets/Views/STableRow.h"
+
+#include "Models/MessagingDebuggerMessageFilter.h"
+#include "Models/MessagingDebuggerModel.h"
 #include "Widgets/History/SMessagingHistoryFilterBar.h"
-#include "Widgets/Views/SListView.h"
 #include "Widgets/History/SMessagingHistoryTableRow.h"
-#include "Widgets/Input/SHyperlink.h"
 
 
 #define LOCTEXT_NAMESPACE "SMessagingHistory"
diff --git a/Engine/Plugins/Messaging/MessagingDebugger/Source/MessagingDebugger/Private/Widgets/History/SMessagingHistory.h b/Engine/Plugins/Messaging/MessagingDebugger/Source/MessagingDebugger/Private/Widgets/History/SMessagingHistory.h
--- a/Engine/Plugins/Messaging/MessagingDebugger/Source/MessagingDebugger/Private/Widgets/History/SMessagingHistory.h
+++ b/Engine/Plugins/Messaging/MessagingDebugger/Source/MessagingDebugger/Private/Widgets/History/SMessagingHistory.h
@@ -2,16 +2,20 @@
 
 #pragma once
 
-#include "CoreMinimal.h"
-#include "Styling/ISlateStyle.h"
-#include "Layout/Visibility.h"
+#include "CoreTypes.h"
+#include "Templates/SharedPointer.h"
 #include "Widgets/DeclarativeSyntaxSupport.h"
 #include "Widgets/SCompoundWidget.h"
-#include "IMessageTracer.h"
-#include "Models/MessagingDebuggerModel.h"
-#include "Widgets/Views/STableViewBase.h"
-#include "Widgets/Views/STableRow.h"
-#include "Models/MessagingDebuggerMessageFilter.h"
+#include "Widgets/Views/SListView.h"
+
+class FMessagingDebuggerMessageFilter;
+class FMessagingDebuggerModel;
+class IMessageTracer;
+class ISlateStyle;
+class STableViewBase;
+
+struct FMessageTracerMessageInfo;
+
 
 /**
  * Implements the message history panel.
@@ -32,7 +36,7 @@ public:
 public:
 
 	/**
-	 * Construct this widget
+	 * Construct this widget.
 	 *
 	 * @param InArgs The construction arguments.
 	 * @param InModel The view model to use.
diff --git a/Engine/Plugins/Messaging/MessagingDebugger/Source/MessagingDebugger/Private/Widgets/History/SMessagingHistoryFilterBar.h b/Engine/Plugins/Messaging/MessagingDebugger/Source/MessagingDebugger/Private/Widgets/History/SMessagingHistoryFilterBar.h
--- a/Engine/Plugins/Messaging/MessagingDebugger/Source/MessagingDebugger/Private/Widgets/History/SMessagingHistoryFilterBar.h
+++ b/Engine/Plugins/Messaging/MessagingDebugger/Source/MessagingDebugger/Private/Widgets/History/SMessagingHistoryFilterBar.h
@@ -2,12 +2,15 @@
 
 #pragma once
 
-#include "CoreMinimal.h"
-#include "Widgets/SNullWidget.h"
+#include "CoreTypes.h"
+#include "Templates/SharedPointer.h"
 #include "Widgets/DeclarativeSyntaxSupport.h"
+#include "Widgets/SNullWidget.h"
 #include "Widgets/SCompoundWidget.h"
+
 #include "Models/MessagingDebuggerMessageFilter.h"
 
+
 #define LOCTEXT_NAMESPACE "SMessagingHistoryFilterBar"
 
 
diff --git a/Engine/Plugins/Messaging/MessagingDebugger/Source/MessagingDebugger/Private/Widgets/History/SMessagingHistoryTableRow.cpp b/Engine/Plugins/Messaging/MessagingDebugger/Source/MessagingDebugger/Private/Widgets/History/SMessagingHistoryTableRow.cpp
--- a/Engine/Plugins/Messaging/MessagingDebugger/Source/MessagingDebugger/Private/Widgets/History/SMessagingHistoryTableRow.cpp
+++ b/Engine/Plugins/Messaging/MessagingDebugger/Source/MessagingDebugger/Private/Widgets/History/SMessagingHistoryTableRow.cpp
@@ -2,6 +2,14 @@
 
 #include "Widgets/History/SMessagingHistoryTableRow.h"
 
+#include "IMessageTracer.h"
+#include "SlateOptMacros.h"
+#include "Styling/ISlateStyle.h"
+#include "Styling/SlateBrush.h"
+#include "Widgets/Images/SImage.h"
+#include "Widgets/Text/STextBlock.h"
+#include "Widgets/Views/STableViewBase.h"
+
 
 #define LOCTEXT_NAMESPACE "SMessagingHistoryTableRow"
 
diff --git a/Engine/Plugins/Messaging/MessagingDebugger/Source/MessagingDebugger/Private/Widgets/History/SMessagingHistoryTableRow.h b/Engine/Plugins/Messaging/MessagingDebugger/Source/MessagingDebugger/Private/Widgets/History/SMessagingHistoryTableRow.h
--- a/Engine/Plugins/Messaging/MessagingDebugger/Source/MessagingDebugger/Private/Widgets/History/SMessagingHistoryTableRow.h
+++ b/Engine/Plugins/Messaging/MessagingDebugger/Source/MessagingDebugger/Private/Widgets/History/SMessagingHistoryTableRow.h
@@ -2,17 +2,24 @@
 
 #pragma once
 
-#include "CoreMinimal.h"
+#include "CoreTypes.h"
+#include "Internationalization/Text.h"
 #include "Misc/Attribute.h"
 #include "Styling/SlateColor.h"
+#include "Templates/SharedPointer.h"
 #include "Widgets/DeclarativeSyntaxSupport.h"
-#include "Widgets/SWidget.h"
-#include "IMessageTracer.h"
-#include "Widgets/Views/STableViewBase.h"
-#include "Widgets/Views/STableRow.h"
-#include "SlateOptMacros.h"
 #include "Widgets/Views/SListView.h"
-#include "Widgets/Images/SImage.h"
+#include "Widgets/Views/STableRow.h"
+
+class IMessageTracer;
+class ISlateStyle;
+class STableViewBase;
+
+struct FMessageTracerMessageInfo;
+struct FSlateBrush;
+
+enum class EMessageScope : uint8;
+
 
 /**
  * Implements a row widget for the message history list.
diff --git a/Engine/Plugins/Messaging/MessagingDebugger/Source/MessagingDebugger/Private/Widgets/Interceptors/SMessagingInterceptorsTableRow.cpp b/Engine/Plugins/Messaging/MessagingDebugger/Source/MessagingDebugger/Private/Widgets/Interceptors/SMessagingInterceptorsTableRow.cpp
--- a/Engine/Plugins/Messaging/MessagingDebugger/Source/MessagingDebugger/Private/Widgets/Interceptors/SMessagingInterceptorsTableRow.cpp
+++ b/Engine/Plugins/Messaging/MessagingDebugger/Source/MessagingDebugger/Private/Widgets/Interceptors/SMessagingInterceptorsTableRow.cpp
@@ -2,6 +2,13 @@
 
 #include "Widgets/Interceptors/SMessagingInterceptorsTableRow.h"
 
+#include "IMessageTracer.h"
+#include "Styling/ISlateStyle.h"
+#include "Widgets/Text/STextBlock.h"
+#include "Widgets/Views/STableViewBase.h"
+
+#include "Models/MessagingDebuggerModel.h"
+
 
 namespace MessagingInterceptorTableRow
 {
diff --git a/Engine/Plugins/Messaging/MessagingDebugger/Source/MessagingDebugger/Private/Widgets/Interceptors/SMessagingInterceptorsTableRow.h b/Engine/Plugins/Messaging/MessagingDebugger/Source/MessagingDebugger/Private/Widgets/Interceptors/SMessagingInterceptorsTableRow.h
--- a/Engine/Plugins/Messaging/MessagingDebugger/Source/MessagingDebugger/Private/Widgets/Interceptors/SMessagingInterceptorsTableRow.h
+++ b/Engine/Plugins/Messaging/MessagingDebugger/Source/MessagingDebugger/Private/Widgets/Interceptors/SMessagingInterceptorsTableRow.h
@@ -2,14 +2,18 @@
 
 #pragma once
 
-#include "CoreMinimal.h"
-#include "Widgets/DeclarativeSyntaxSupport.h"
-#include "Widgets/SWidget.h"
-#include "Models/MessagingDebuggerModel.h"
-#include "Widgets/Views/STableViewBase.h"
-#include "Widgets/Views/STableRow.h"
+#include "CoreTypes.h"
 #include "SlateOptMacros.h"
+#include "Widgets/DeclarativeSyntaxSupport.h"
 #include "Widgets/Views/SListView.h"
+#include "Widgets/Views/STableRow.h"
+
+class FMessagingDebuggerModel;
+class ISlateStyle;
+class STableViewBase;
+
+struct FMessageTracerInterceptorInfo;
+
 
 /**
  * Implements a row widget for the interceptors list.
diff --git a/Engine/Plugins/Messaging/MessagingDebugger/Source/MessagingDebugger/Private/Widgets/MessageData/SMessagingMessageData.cpp b/Engine/Plugins/Messaging/MessagingDebugger/Source/MessagingDebugger/Private/Widgets/MessageData/SMessagingMessageData.cpp
--- a/Engine/Plugins/Messaging/MessagingDebugger/Source/MessagingDebugger/Private/Widgets/MessageData/SMessagingMessageData.cpp
+++ b/Engine/Plugins/Messaging/MessagingDebugger/Source/MessagingDebugger/Private/Widgets/MessageData/SMessagingMessageData.cpp
@@ -1,12 +1,14 @@
 // Copyright 1998-2017 Epic Games, Inc. All Rights Reserved.
 
 #include "Widgets/MessageData/SMessagingMessageData.h"
-#include "Serialization/BufferArchive.h"
+
+#include "Backends/JsonStructSerializerBackend.h"
 #include "Modules/ModuleManager.h"
+#include "Serialization/BufferArchive.h"
+#include "StructSerializer.h"
+#include "Styling/ISlateStyle.h"
 #include "UObject/StructOnScope.h"
 #include "Widgets/Input/SMultiLineEditableTextBox.h"
-#include "Backends/JsonStructSerializerBackend.h"
-#include "StructSerializer.h"
 
 #if WITH_EDITOR
 	#include "IDetailsView.h"
@@ -14,6 +16,8 @@
 	#include "PropertyEditorModule.h"
 #endif
 
+#include "Models/MessagingDebuggerModel.h"
+
 
 #define LOCTEXT_NAMESPACE "SMessagingMessageData"
 
diff --git a/Engine/Plugins/Messaging/MessagingDebugger/Source/MessagingDebugger/Private/Widgets/MessageData/SMessagingMessageData.h b/Engine/Plugins/Messaging/MessagingDebugger/Source/MessagingDebugger/Private/Widgets/MessageData/SMessagingMessageData.h
--- a/Engine/Plugins/Messaging/MessagingDebugger/Source/MessagingDebugger/Private/Widgets/MessageData/SMessagingMessageData.h
+++ b/Engine/Plugins/Messaging/MessagingDebugger/Source/MessagingDebugger/Private/Widgets/MessageData/SMessagingMessageData.h
@@ -2,18 +2,20 @@
 
 #pragma once
 
-#include "CoreMinimal.h"
-#include "Styling/ISlateStyle.h"
-#include "Layout/Visibility.h"
+#include "CoreTypes.h"
+#include "Misc/NotifyHook.h"
+#include "Templates/SharedPointer.h"
 #include "Widgets/DeclarativeSyntaxSupport.h"
 #include "Widgets/SCompoundWidget.h"
-#include "Models/MessagingDebuggerModel.h"
-#include "Misc/NotifyHook.h"
 
+class FMessagingDebuggerModel;
+class ISlateStyle;
 class IStructureDetailsView;
 class SMultiLineEditableTextBox;
+
 struct FPropertyChangedEvent;
 
+
 /**
  * Implements the message data panel.
  */
@@ -34,7 +36,7 @@ public:
 public:
 
 	/**
-	 * Construct this widget
+	 * Construct this widget.
 	 *
 	 * @param InArgs The declaration data for this widget.
 	 * @param InModel The view model to use.
diff --git a/Engine/Plugins/Messaging/MessagingDebugger/Source/MessagingDebugger/Private/Widgets/MessageDetails/SMessagingDispatchStateTableRow.h b/Engine/Plugins/Messaging/MessagingDebugger/Source/MessagingDebugger/Private/Widgets/MessageDetails/SMessagingDispatchStateTableRow.h
--- a/Engine/Plugins/Messaging/MessagingDebugger/Source/MessagingDebugger/Private/Widgets/MessageDetails/SMessagingDispatchStateTableRow.h
+++ b/Engine/Plugins/Messaging/MessagingDebugger/Source/MessagingDebugger/Private/Widgets/MessageDetails/SMessagingDispatchStateTableRow.h
@@ -3,24 +3,27 @@
 #pragma once
 
 #include "CoreMinimal.h"
-#include "Styling/SlateColor.h"
+#include "IMessageTracer.h"
 #include "Layout/Margin.h"
-#include "Widgets/SNullWidget.h"
+#include "SlateOptMacros.h"
+#include "Styling/SlateColor.h"
 #include "Widgets/DeclarativeSyntaxSupport.h"
-#include "Widgets/SWidget.h"
-#include "IMessageTracer.h"
-#include "Models/MessagingDebuggerModel.h"
-#include "Widgets/Views/STableViewBase.h"
 #include "Widgets/SBoxPanel.h"
+#include "Widgets/SNullWidget.h"
+#include "Widgets/SWidget.h"
+#include "Widgets/Images/SImage.h"
 #include "Widgets/Layout/SBox.h"
-#include "Widgets/Views/STableRow.h"
 #include "Widgets/Text/STextBlock.h"
 #include "Widgets/Views/SListView.h"
-#include "SlateOptMacros.h"
-#include "Widgets/Images/SImage.h"
+#include "Widgets/Views/STableRow.h"
+#include "Widgets/Views/STableViewBase.h"
+
+#include "Models/MessagingDebuggerModel.h"
+
 
 #define LOCTEXT_NAMESPACE "SMessagingDispatchStateTableRow"
 
+
 /**
  * Implements a row widget for the dispatch state list.
  */
diff --git a/Engine/Plugins/Messaging/MessagingDebugger/Source/MessagingDebugger/Private/Widgets/MessageDetails/SMessagingMessageDetails.cpp b/Engine/Plugins/Messaging/MessagingDebugger/Source/MessagingDebugger/Private/Widgets/MessageDetails/SMessagingMessageDetails.cpp
--- a/Engine/Plugins/Messaging/MessagingDebugger/Source/MessagingDebugger/Private/Widgets/MessageDetails/SMessagingMessageDetails.cpp
+++ b/Engine/Plugins/Messaging/MessagingDebugger/Source/MessagingDebugger/Private/Widgets/MessageDetails/SMessagingMessageDetails.cpp
@@ -1,9 +1,14 @@
 // Copyright 1998-2017 Epic Games, Inc. All Rights Reserved.
 
 #include "Widgets/MessageDetails/SMessagingMessageDetails.h"
-#include "Widgets/Text/STextBlock.h"
+
+#include "IMessageTracer.h"
+#include "Styling/ISlateStyle.h"
 #include "Widgets/Layout/SGridPanel.h"
-#include "Widgets/Views/SListView.h"
+#include "Widgets/Text/STextBlock.h"
+#include "Widgets/Views/STableViewBase.h"
+
+#include "Models/MessagingDebuggerModel.h"
 #include "Widgets/MessageDetails/SMessagingDispatchStateTableRow.h"
 
 
diff --git a/Engine/Plugins/Messaging/MessagingDebugger/Source/MessagingDebugger/Private/Widgets/MessageDetails/SMessagingMessageDetails.h b/Engine/Plugins/Messaging/MessagingDebugger/Source/MessagingDebugger/Private/Widgets/MessageDetails/SMessagingMessageDetails.h
--- a/Engine/Plugins/Messaging/MessagingDebugger/Source/MessagingDebugger/Private/Widgets/MessageDetails/SMessagingMessageDetails.h
+++ b/Engine/Plugins/Messaging/MessagingDebugger/Source/MessagingDebugger/Private/Widgets/MessageDetails/SMessagingMessageDetails.h
@@ -2,13 +2,19 @@
 
 #pragma once
 
-#include "CoreMinimal.h"
-#include "Styling/ISlateStyle.h"
+#include "CoreTypes.h"
+#include "Templates/SharedPointer.h"
 #include "Widgets/DeclarativeSyntaxSupport.h"
 #include "Widgets/SCompoundWidget.h"
-#include "Models/MessagingDebuggerModel.h"
-#include "Widgets/Views/STableViewBase.h"
-#include "Widgets/Views/STableRow.h"
+#include "Widgets/Views/SListView.h"
+
+class FMessagingDebuggerModel;
+class ISlateStyle;
+class ITableRow;
+class STableViewBase;
+
+struct FMessageTracerDispatchState;
+
 
 /**
  * Implements the message details panel.
@@ -29,7 +35,7 @@ public:
 public:
 
 	/**
-	 * Construct this widget
+	 * Construct this widget.
 	 *
 	 * @param InArgs The declaration data for this widget.
 	 * @param InModel The view model to use.
diff --git a/Engine/Plugins/Messaging/MessagingDebugger/Source/MessagingDebugger/Private/Widgets/SMessagingDebugger.cpp b/Engine/Plugins/Messaging/MessagingDebugger/Source/MessagingDebugger/Private/Widgets/SMessagingDebugger.cpp
--- a/Engine/Plugins/Messaging/MessagingDebugger/Source/MessagingDebugger/Private/Widgets/SMessagingDebugger.cpp
+++ b/Engine/Plugins/Messaging/MessagingDebugger/Source/MessagingDebugger/Private/Widgets/SMessagingDebugger.cpp
@@ -1,24 +1,28 @@
 // Copyright 1998-2017 Epic Games, Inc. All Rights Reserved.
 
 #include "Widgets/SMessagingDebugger.h"
-#include "Widgets/SBoxPanel.h"
-#include "Textures/SlateIcon.h"
+
+#include "IMessageTracer.h"
 #include "Framework/Commands/UICommandList.h"
 #include "Framework/MultiBox/MultiBoxBuilder.h"
 #include "Framework/Docking/WorkspaceItem.h"
 #include "Framework/Docking/TabManager.h"
+#include "Textures/SlateIcon.h"
+#include "Widgets/SBoxPanel.h"
+#include "Widgets/Docking/SDockTab.h"
+
 #include "Models/MessagingDebuggerCommands.h"
+#include "Models/MessagingDebuggerModel.h"
 #include "Widgets/Breakpoints/SMessagingBreakpoints.h"
-#include "Widgets/Toolbar/SMessagingDebuggerToolbar.h"
 #include "Widgets/EndpointDetails/SMessagingEndpointDetails.h"
 #include "Widgets/Endpoints/SMessagingEndpoints.h"
 #include "Widgets/Graph/SMessagingGraph.h"
-#include "Widgets/MessageData/SMessagingMessageData.h"
-#include "Widgets/MessageDetails/SMessagingMessageDetails.h"
 #include "Widgets/History/SMessagingHistory.h"
 #include "Widgets/Interceptors/SMessagingInterceptors.h"
+#include "Widgets/MessageData/SMessagingMessageData.h"
+#include "Widgets/MessageDetails/SMessagingMessageDetails.h"
+#include "Widgets/Toolbar/SMessagingDebuggerToolbar.h"
 #include "Widgets/Types/SMessagingTypes.h"
-#include "Widgets/Docking/SDockTab.h"
 
 
 #define LOCTEXT_NAMESPACE "SMessagingDebugger"
diff --git a/Engine/Plugins/Messaging/MessagingDebugger/Source/MessagingDebugger/Private/Widgets/SMessagingDebugger.h b/Engine/Plugins/Messaging/MessagingDebugger/Source/MessagingDebugger/Private/Widgets/SMessagingDebugger.h
--- a/Engine/Plugins/Messaging/MessagingDebugger/Source/MessagingDebugger/Private/Widgets/SMessagingDebugger.h
+++ b/Engine/Plugins/Messaging/MessagingDebugger/Source/MessagingDebugger/Private/Widgets/SMessagingDebugger.h
@@ -8,15 +8,16 @@
 #include "Styling/ISlateStyle.h"
 #include "Widgets/DeclarativeSyntaxSupport.h"
 #include "Widgets/SCompoundWidget.h"
-#include "IMessageTracer.h"
-#include "Models/MessagingDebuggerModel.h"
 
 class FMenuBuilder;
+class FMessagingDebuggerModel;
 class FSpawnTabArgs;
 class FTabManager;
 class FUICommandList;
+class IMessageTracer;
 class SWindow;
 
+
 /**
  * Implements the messaging system debugger application.
  */
@@ -54,11 +55,11 @@ public:
 protected:
 
 	/**
-	* Fills the Window menu with menu items.
-	*
-	* @param MenuBuilder The multi-box builder that should be filled with content for this pull-down menu.
-	* @param TabManager A Tab Manager from which to populate tab spawner menu items.
-	*/
+	 * Fills the Window menu with menu items.
+	 *
+	 * @param MenuBuilder The multi-box builder that should be filled with content for this pull-down menu.
+	 * @param TabManager A Tab Manager from which to populate tab spawner menu items.
+	 */
 	static void FillWindowMenu(FMenuBuilder& MenuBuilder, const TSharedPtr<FTabManager> TabManager);
 
 protected:
diff --git a/Engine/Plugins/Messaging/MessagingDebugger/Source/MessagingDebugger/Private/Widgets/Toolbar/SMessagingDebuggerToolbar.cpp b/Engine/Plugins/Messaging/MessagingDebugger/Source/MessagingDebugger/Private/Widgets/Toolbar/SMessagingDebuggerToolbar.cpp
--- a/Engine/Plugins/Messaging/MessagingDebugger/Source/MessagingDebugger/Private/Widgets/Toolbar/SMessagingDebuggerToolbar.cpp
+++ b/Engine/Plugins/Messaging/MessagingDebugger/Source/MessagingDebugger/Private/Widgets/Toolbar/SMessagingDebuggerToolbar.cpp
@@ -1,8 +1,13 @@
 // Copyright 1998-2017 Epic Games, Inc. All Rights Reserved.
 
 #include "Widgets/Toolbar/SMessagingDebuggerToolbar.h"
+
+#include "Framework/Commands/UICommandList.h"
 #include "Framework/MultiBox/MultiBoxDefs.h"
 #include "Framework/MultiBox/MultiBoxBuilder.h"
+#include "Styling/ISlateStyle.h"
+#include "Widgets/SWidget.h"
+
 #include "Models/MessagingDebuggerCommands.h"
 
 
diff --git a/Engine/Plugins/Messaging/MessagingDebugger/Source/MessagingDebugger/Private/Widgets/Toolbar/SMessagingDebuggerToolbar.h b/Engine/Plugins/Messaging/MessagingDebugger/Source/MessagingDebugger/Private/Widgets/Toolbar/SMessagingDebuggerToolbar.h
--- a/Engine/Plugins/Messaging/MessagingDebugger/Source/MessagingDebugger/Private/Widgets/Toolbar/SMessagingDebuggerToolbar.h
+++ b/Engine/Plugins/Messaging/MessagingDebugger/Source/MessagingDebugger/Private/Widgets/Toolbar/SMessagingDebuggerToolbar.h
@@ -2,12 +2,15 @@
 
 #pragma once
 
-#include "CoreMinimal.h"
-#include "Styling/ISlateStyle.h"
+#include "CoreTypes.h"
+#include "Templates/SharedPointer.h"
 #include "Widgets/DeclarativeSyntaxSupport.h"
-#include "Widgets/SWidget.h"
 #include "Widgets/SCompoundWidget.h"
-#include "Framework/Commands/UICommandList.h"
+
+class FUICommandList;
+class ISlateStyle;
+class SWidget;
+
 
 /**
  * Implements the messaging debugger toolbar widget.
@@ -23,13 +26,13 @@ public:
 public:
 
 	/**
-	 * Construct this widget
+	 * Construct this widget.
 	 *
 	 * @param InArgs The declaration data for this widget.
 	 * @param InStyle The visual style to use for this widget.
 	 * @param InCommandList The command list to bind to.
 	 */
-	void Construct( const FArguments& InArgs, const TSharedRef<ISlateStyle>& InStyle, const TSharedRef<FUICommandList>& InCommandList );
+	void Construct(const FArguments& InArgs, const TSharedRef<ISlateStyle>& InStyle, const TSharedRef<FUICommandList>& InCommandList);
 
 protected:
 
@@ -39,5 +42,5 @@ protected:
 	 * @param CommandList The command list to use for the toolbar buttons.
 	 * @return The toolbar widget.
 	 */
-	TSharedRef<SWidget> MakeToolbar( const TSharedRef<FUICommandList>& CommandList );
+	TSharedRef<SWidget> MakeToolbar(const TSharedRef<FUICommandList>& CommandList);
 };
diff --git a/Engine/Plugins/Messaging/MessagingDebugger/Source/MessagingDebugger/Private/Widgets/Types/SMessagingTypes.cpp b/Engine/Plugins/Messaging/MessagingDebugger/Source/MessagingDebugger/Private/Widgets/Types/SMessagingTypes.cpp
--- a/Engine/Plugins/Messaging/MessagingDebugger/Source/MessagingDebugger/Private/Widgets/Types/SMessagingTypes.cpp
+++ b/Engine/Plugins/Messaging/MessagingDebugger/Source/MessagingDebugger/Private/Widgets/Types/SMessagingTypes.cpp
@@ -1,11 +1,17 @@
 // Copyright 1998-2017 Epic Games, Inc. All Rights Reserved.
 
 #include "Widgets/Types/SMessagingTypes.h"
+
+#include "IMessageTracer.h"
+#include "Styling/ISlateStyle.h"
 #include "Widgets/Images/SImage.h"
+#include "Widgets/Layout/SExpandableArea.h"
+#include "Widgets/Views/STableViewBase.h"
+
+#include "Models/MessagingDebuggerModel.h"
+#include "Models/MessagingDebuggerTypeFilter.h"
 #include "Widgets/Types/SMessagingTypesFilterBar.h"
-#include "Widgets/Views/SListView.h"
 #include "Widgets/Types/SMessagingTypesTableRow.h"
-#include "Widgets/Layout/SExpandableArea.h"
 
 
 #define LOCTEXT_NAMESPACE "SMessagingTypes"
diff --git a/Engine/Plugins/Messaging/MessagingDebugger/Source/MessagingDebugger/Private/Widgets/Types/SMessagingTypes.h b/Engine/Plugins/Messaging/MessagingDebugger/Source/MessagingDebugger/Private/Widgets/Types/SMessagingTypes.h
--- a/Engine/Plugins/Messaging/MessagingDebugger/Source/MessagingDebugger/Private/Widgets/Types/SMessagingTypes.h
+++ b/Engine/Plugins/Messaging/MessagingDebugger/Source/MessagingDebugger/Private/Widgets/Types/SMessagingTypes.h
@@ -2,15 +2,22 @@
 
 #pragma once
 
-#include "CoreMinimal.h"
-#include "Styling/ISlateStyle.h"
+#include "CoreTypes.h"
+#include "Containers/Array.h"
+#include "Templates/SharedPointer.h"
 #include "Widgets/DeclarativeSyntaxSupport.h"
 #include "Widgets/SCompoundWidget.h"
-#include "IMessageTracer.h"
-#include "Models/MessagingDebuggerModel.h"
-#include "Widgets/Views/STableViewBase.h"
-#include "Widgets/Views/STableRow.h"
-#include "Models/MessagingDebuggerTypeFilter.h"
+#include "Widgets/Views/SListView.h"
+
+class FMessagingDebuggerModel;
+class FMessagingDebuggerTypeFilter;
+class IMessageTracer;
+class ISlateStyle;
+class ITableRow;
+class STableViewBase;
+
+struct FMessageTracerTypeInfo;
+
 
 /**
  * Implements the message types panel.
@@ -31,7 +38,7 @@ public:
 public:
 
 	/**
-	 * Construct this widget
+	 * Construct this widget.
 	 *
 	 * @param InArgs The declaration data for this widget.
 	 * @param InModel The view model to use.
diff --git a/Engine/Plugins/Messaging/MessagingDebugger/Source/MessagingDebugger/Private/Widgets/Types/SMessagingTypesFilterBar.cpp b/Engine/Plugins/Messaging/MessagingDebugger/Source/MessagingDebugger/Private/Widgets/Types/SMessagingTypesFilterBar.cpp
--- a/Engine/Plugins/Messaging/MessagingDebugger/Source/MessagingDebugger/Private/Widgets/Types/SMessagingTypesFilterBar.cpp
+++ b/Engine/Plugins/Messaging/MessagingDebugger/Source/MessagingDebugger/Private/Widgets/Types/SMessagingTypesFilterBar.cpp
@@ -1,9 +1,12 @@
 // Copyright 1998-2017 Epic Games, Inc. All Rights Reserved.
 
 #include "Widgets/Types/SMessagingTypesFilterBar.h"
+
 #include "Widgets/SBoxPanel.h"
 #include "Widgets/Input/SSearchBox.h"
 
+#include "Models/MessagingDebuggerTypeFilter.h"
+
 
 #define LOCTEXT_NAMESPACE "SMessagingTypesFilterBar"
 
diff --git a/Engine/Plugins/Messaging/MessagingDebugger/Source/MessagingDebugger/Private/Widgets/Types/SMessagingTypesFilterBar.h b/Engine/Plugins/Messaging/MessagingDebugger/Source/MessagingDebugger/Private/Widgets/Types/SMessagingTypesFilterBar.h
--- a/Engine/Plugins/Messaging/MessagingDebugger/Source/MessagingDebugger/Private/Widgets/Types/SMessagingTypesFilterBar.h
+++ b/Engine/Plugins/Messaging/MessagingDebugger/Source/MessagingDebugger/Private/Widgets/Types/SMessagingTypesFilterBar.h
@@ -2,10 +2,13 @@
 
 #pragma once
 
-#include "CoreMinimal.h"
+#include "CoreTypes.h"
+#include "Templates/SharedPointer.h"
 #include "Widgets/DeclarativeSyntaxSupport.h"
 #include "Widgets/SCompoundWidget.h"
-#include "Models/MessagingDebuggerTypeFilter.h"
+
+class FMessagingDebuggerTypeFilter;
+
 
 /**
  * Implements the message type list filter bar widget.
@@ -25,12 +28,12 @@ public:
 public:
 
 	/**
-	 * Construct this widget
+	 * Construct this widget.
 	 *
 	 * @param InArgs The declaration data for this widget.
 	 * @param InFilter The filter model.
 	 */
-	void Construct( const FArguments& InArgs, TSharedRef<FMessagingDebuggerTypeFilter> InFilter );
+	void Construct(const FArguments& InArgs, TSharedRef<FMessagingDebuggerTypeFilter> InFilter);
 
 private:
 
diff --git a/Engine/Plugins/Messaging/MessagingDebugger/Source/MessagingDebugger/Private/Widgets/Types/SMessagingTypesTableRow.cpp b/Engine/Plugins/Messaging/MessagingDebugger/Source/MessagingDebugger/Private/Widgets/Types/SMessagingTypesTableRow.cpp
--- a/Engine/Plugins/Messaging/MessagingDebugger/Source/MessagingDebugger/Private/Widgets/Types/SMessagingTypesTableRow.cpp
+++ b/Engine/Plugins/Messaging/MessagingDebugger/Source/MessagingDebugger/Private/Widgets/Types/SMessagingTypesTableRow.cpp
@@ -2,6 +2,10 @@
 
 #include "Widgets/Types/SMessagingTypesTableRow.h"
 
+#include "IMessageTracer.h"
+
+#include "Models/MessagingDebuggerModel.h"
+
 
 #define LOCTEXT_NAMESPACE "SMessagingTypesTableRow"
 
diff --git a/Engine/Plugins/Messaging/MessagingDebugger/Source/MessagingDebugger/Private/Widgets/Types/SMessagingTypesTableRow.h b/Engine/Plugins/Messaging/MessagingDebugger/Source/MessagingDebugger/Private/Widgets/Types/SMessagingTypesTableRow.h
--- a/Engine/Plugins/Messaging/MessagingDebugger/Source/MessagingDebugger/Private/Widgets/Types/SMessagingTypesTableRow.h
+++ b/Engine/Plugins/Messaging/MessagingDebugger/Source/MessagingDebugger/Private/Widgets/Types/SMessagingTypesTableRow.h
@@ -2,16 +2,22 @@
 
 #pragma once
 
-#include "CoreMinimal.h"
+#include "CoreTypes.h"
+#include "Internationalization/Text.h"
 #include "Misc/Attribute.h"
+#include "SlateOptMacros.h"
+#include "Templates/SharedPointer.h"
 #include "Widgets/DeclarativeSyntaxSupport.h"
 #include "Widgets/SWidget.h"
-#include "Models/MessagingDebuggerModel.h"
+#include "Widgets/Input/SCheckBox.h"
 #include "Widgets/Views/STableViewBase.h"
 #include "Widgets/Views/STableRow.h"
 #include "Widgets/Views/SListView.h"
-#include "SlateOptMacros.h"
-#include "Widgets/Input/SCheckBox.h"
+
+class FMessagingDebuggerModel;
+
+struct FMessageTracerTypeInfo;
+
 
 /**
  * Implements a row widget for the message type list.
diff --git a/Engine/Plugins/Messaging/TcpMessaging/Source/TcpMessaging/Private/TcpMessagingModule.cpp b/Engine/Plugins/Messaging/TcpMessaging/Source/TcpMessaging/Private/TcpMessagingModule.cpp
--- a/Engine/Plugins/Messaging/TcpMessaging/Source/TcpMessaging/Private/TcpMessagingModule.cpp
+++ b/Engine/Plugins/Messaging/TcpMessaging/Source/TcpMessaging/Private/TcpMessagingModule.cpp
@@ -1,21 +1,24 @@
 // Copyright 1998-2017 Epic Games, Inc. All Rights Reserved.
 
+#include "TcpMessagingPrivate.h"
+
 #include "CoreMinimal.h"
+#include "MessageBridgeBuilder.h"
 #include "Misc/CoreMisc.h"
 #include "Misc/CommandLine.h"
 #include "Misc/CoreDelegates.h"
 #include "Misc/App.h"
 #include "Modules/ModuleManager.h"
-#include "Helpers/MessageBridgeBuilder.h"
 #include "Interfaces/IPv4/IPv4Endpoint.h"
-#include "TcpMessagingPrivate.h"
 #include "Settings/TcpMessagingSettings.h"
-#include "Transport/TcpMessageTransport.h"
+
 #if WITH_EDITOR
-#include "ISettingsModule.h"
-#include "ISettingsSection.h"
+	#include "ISettingsModule.h"
+	#include "ISettingsSection.h"
 #endif
+
 #include "ITcpMessagingModule.h"
+#include "Transport/TcpMessageTransport.h"
 
 
 DEFINE_LOG_CATEGORY(LogTcpMessaging);
diff --git a/Engine/Plugins/Messaging/TcpMessaging/Source/TcpMessaging/Private/Transport/TcpDeserializedMessage.cpp b/Engine/Plugins/Messaging/TcpMessaging/Source/TcpMessaging/Private/Transport/TcpDeserializedMessage.cpp
--- a/Engine/Plugins/Messaging/TcpMessaging/Source/TcpMessaging/Private/Transport/TcpDeserializedMessage.cpp
+++ b/Engine/Plugins/Messaging/TcpMessaging/Source/TcpMessaging/Private/Transport/TcpDeserializedMessage.cpp
@@ -34,7 +34,7 @@ FTcpDeserializedMessage::~FTcpDeserializedMessage()
 /* FTcpDeserializedMessage interface
  *****************************************************************************/
 
-bool FTcpDeserializedMessage::Deserialize(const FArrayReaderPtr& Message)
+bool FTcpDeserializedMessage::Deserialize(const TSharedPtr<FArrayReader, ESPMode::ThreadSafe>& Message)
 {
 	FArrayReader& MessageReader = Message.ToSharedRef().Get();
 
@@ -161,7 +161,7 @@ const TWeakObjectPtr<UScriptStruct>& FTcpDeserializedMessage::GetMessageTypeInfo
 }
 
 
-IMessageContextPtr FTcpDeserializedMessage::GetOriginalContext() const
+TSharedPtr<IMessageContext, ESPMode::ThreadSafe> FTcpDeserializedMessage::GetOriginalContext() const
 {
 	return nullptr;
 }
diff --git a/Engine/Plugins/Messaging/TcpMessaging/Source/TcpMessaging/Private/Transport/TcpDeserializedMessage.h b/Engine/Plugins/Messaging/TcpMessaging/Source/TcpMessaging/Private/Transport/TcpDeserializedMessage.h
--- a/Engine/Plugins/Messaging/TcpMessaging/Source/TcpMessaging/Private/Transport/TcpDeserializedMessage.h
+++ b/Engine/Plugins/Messaging/TcpMessaging/Source/TcpMessaging/Private/Transport/TcpDeserializedMessage.h
@@ -33,7 +33,7 @@ public:
 	 * @param ReassembledMessage The reassembled message to deserialize.
 	 * @return true on success, false otherwise.
 	 */
-	bool Deserialize(const FArrayReaderPtr& Message);
+	bool Deserialize(const TSharedPtr<FArrayReader, ESPMode::ThreadSafe>& Message);
 
 public:
 
@@ -44,7 +44,7 @@ public:
 	virtual const FDateTime& GetExpiration() const override;
 	virtual const void* GetMessage() const override;
 	virtual const TWeakObjectPtr<UScriptStruct>& GetMessageTypeInfo() const override;
-	virtual IMessageContextPtr GetOriginalContext() const override;
+	virtual TSharedPtr<IMessageContext, ESPMode::ThreadSafe> GetOriginalContext() const override;
 	virtual const TArray<FMessageAddress>& GetRecipients() const override;
 	virtual EMessageScope GetScope() const override;
 	virtual const FMessageAddress& GetSender() const override;
diff --git a/Engine/Plugins/Messaging/TcpMessaging/Source/TcpMessaging/Private/Transport/TcpMessageTransport.cpp b/Engine/Plugins/Messaging/TcpMessaging/Source/TcpMessaging/Private/Transport/TcpMessageTransport.cpp
--- a/Engine/Plugins/Messaging/TcpMessaging/Source/TcpMessaging/Private/Transport/TcpMessageTransport.cpp
+++ b/Engine/Plugins/Messaging/TcpMessaging/Source/TcpMessaging/Private/Transport/TcpMessageTransport.cpp
@@ -1,10 +1,13 @@
 // Copyright 1998-2017 Epic Games, Inc. All Rights Reserved.
 
 #include "Transport/TcpMessageTransport.h"
+
 #include "HAL/RunnableThread.h"
 #include "Common/TcpSocketBuilder.h"
 #include "Common/TcpListener.h"
+#include "IMessageTransportHandler.h"
 #include "TcpMessagingPrivate.h"
+
 #include "Transport/TcpDeserializedMessage.h"
 #include "Transport/TcpSerializedMessage.h"
 #include "Transport/TcpMessageTransportConnection.h"
@@ -21,6 +24,7 @@ FTcpMessageTransport::FTcpMessageTransport(const FIPv4Endpoint& InListenEndpoint
 	, bStopping(false)
 	, SocketSubsystem(ISocketSubsystem::Get(PLATFORM_SOCKETSUBSYSTEM))
 	, Listener(nullptr)
+	, TransportHandler(nullptr)
 {
 	Thread = FRunnableThread::Create(this, TEXT("FTcpMessageTransport"), 128 * 1024, TPri_Normal);
 }
@@ -41,8 +45,10 @@ FTcpMessageTransport::~FTcpMessageTransport()
 /* IMessageTransport interface
  *****************************************************************************/
 
-bool FTcpMessageTransport::StartTransport()
+bool FTcpMessageTransport::StartTransport(IMessageTransportHandler& Handler)
 {
+	TransportHandler = &Handler;
+
 	if (ListenEndpoint != FIPv4Endpoint::Any)
 	{
 		Listener = new FTcpListener(ListenEndpoint);
@@ -61,16 +67,19 @@ bool FTcpMessageTransport::StartTransport()
 void FTcpMessageTransport::AddOutgoingConnection(const FIPv4Endpoint& Endpoint)
 {
 	FSocket* Socket = FTcpSocketBuilder(TEXT("FTcpMessageTransport.RemoteConnection"));
-	if (Socket)
+	
+	if (Socket == nullptr)
 	{
-		if (!Socket->Connect(*Endpoint.ToInternetAddr()))
-		{
-			ISocketSubsystem::Get(PLATFORM_SOCKETSUBSYSTEM)->DestroySocket(Socket);
-		}
-		else
-		{
-			PendingConnections.Enqueue(MakeShareable(new FTcpMessageTransportConnection(Socket, Endpoint, ConnectionRetryDelay)));
-		}
+		return;
+	}
+
+	if (!Socket->Connect(*Endpoint.ToInternetAddr()))
+	{
+		ISocketSubsystem::Get(PLATFORM_SOCKETSUBSYSTEM)->DestroySocket(Socket);
+	}
+	else
+	{
+		PendingConnections.Enqueue(MakeShareable(new FTcpMessageTransportConnection(Socket, Endpoint, ConnectionRetryDelay)));
 	}
 }
 
@@ -93,13 +102,16 @@ void FTcpMessageTransport::StopTransport()
 	{
 		Connection->Close();
 	}
+
 	Connections.Empty();
 	PendingConnections.Empty();
 	ConnectionEndpointsToRemove.Empty();
+
+	TransportHandler = nullptr;
 }
 
 
-bool FTcpMessageTransport::TransportMessage(const IMessageContextRef& Context, const TArray<FGuid>& Recipients)
+bool FTcpMessageTransport::TransportMessage(const TSharedRef<IMessageContext, ESPMode::ThreadSafe>& Context, const TArray<FGuid>& Recipients)
 {
 	if (Context->GetRecipients().Num() > TCP_MESSAGING_MAX_RECIPIENTS)
 	{
@@ -167,11 +179,18 @@ bool FTcpMessageTransport::TransportMessage(const IMessageContextRef& Context, c
 /* FRunnable interface
 *****************************************************************************/
 
+void FTcpMessageTransport::Exit()
+{
+	// do nothing
+}
+
+
 bool FTcpMessageTransport::Init()
 {
 	return true;
 }
 
+
 uint32 FTcpMessageTransport::Run()
 {
 	while (!bStopping)
@@ -195,6 +214,7 @@ uint32 FTcpMessageTransport::Run()
 				for (int32 Index = 0; Index < Connections.Num(); Index++)
 				{
 					auto& Connection = Connections[Index];
+
 					if (Connection->GetRemoteEndpoint() == Endpoint)
 					{
 						Connection->Close();
@@ -215,10 +235,12 @@ uint32 FTcpMessageTransport::Run()
 			case FTcpMessageTransportConnection::STATE_Connected:
 				ActiveConnections++;
 				break;
+
 			case FTcpMessageTransportConnection::STATE_Disconnected:
 				Connections.RemoveAtSwap(Index);
 				Index--;
 				break;
+
 			default:
 				break;
 			}
@@ -230,10 +252,11 @@ uint32 FTcpMessageTransport::Run()
 			{
 				TSharedPtr<FTcpDeserializedMessage, ESPMode::ThreadSafe> Message;
 				FGuid SenderNodeId;
+
 				while (Connection->Receive(Message, SenderNodeId))
 				{
 					UE_LOG(LogTcpMessaging, Verbose, TEXT("Received message '%s'"), *Message->GetMessageType().ToString());
-					MessageReceivedDelegate.ExecuteIfBound(Message.ToSharedRef(), SenderNodeId);
+					TransportHandler->ReceiveTransportMessage(Message.ToSharedRef(), SenderNodeId);
 				}
 			}
 		}
@@ -244,46 +267,42 @@ uint32 FTcpMessageTransport::Run()
 	return 0;
 }
 
+
 void FTcpMessageTransport::Stop()
 {
 	bStopping = true;
 }
 
-void FTcpMessageTransport::Exit()
-{
-	// do nothing
-}
-
 
-
-/* FTcpMessageTransport event handlers
+/* FTcpMessageTransport callbacks
 *****************************************************************************/
 
 bool FTcpMessageTransport::HandleListenerConnectionAccepted(FSocket* ClientSocket, const FIPv4Endpoint& ClientEndpoint)
 {
 	PendingConnections.Enqueue(MakeShareable(new FTcpMessageTransportConnection(ClientSocket, ClientEndpoint, 0)));
+	
 	return true;
 }
 
+
 void FTcpMessageTransport::HandleConnectionStateChanged(TSharedPtr<FTcpMessageTransportConnection> Connection)
 {
-	FIPv4Endpoint RemoteEndpoint = Connection->GetRemoteEndpoint();
-	FGuid NodeId = Connection->GetRemoteNodeId();
-	FTcpMessageTransportConnection::EConnectionState State = Connection->GetConnectionState();
+	const FGuid NodeId = Connection->GetRemoteNodeId();
+	const FIPv4Endpoint RemoteEndpoint = Connection->GetRemoteEndpoint();
+	const FTcpMessageTransportConnection::EConnectionState State = Connection->GetConnectionState();
 
 	if (State == FTcpMessageTransportConnection::STATE_Connected)
 	{
 		NodeConnectionMapUpdates.Enqueue(FNodeConnectionMapUpdate(true, NodeId, TWeakPtr<FTcpMessageTransportConnection>(Connection)));
-		NodeDiscoveredDelegate.ExecuteIfBound(NodeId);
+		TransportHandler->DiscoverTransportNode(NodeId);
+
 		UE_LOG(LogTcpMessaging, Verbose, TEXT("Discovered node '%s' on connection '%s'..."), *NodeId.ToString(), *RemoteEndpoint.ToString());
 	}
-	else
+	else if (NodeId.IsValid())
 	{
-		if (NodeId.IsValid())
-		{
-			UE_LOG(LogTcpMessaging, Verbose, TEXT("Lost node '%s' on connection '%s'..."), *NodeId.ToString(), *RemoteEndpoint.ToString());
-			NodeConnectionMapUpdates.Enqueue(FNodeConnectionMapUpdate(false, NodeId, TWeakPtr<FTcpMessageTransportConnection>(Connection)));
-			NodeLostDelegate.ExecuteIfBound(NodeId);
-		}
+		UE_LOG(LogTcpMessaging, Verbose, TEXT("Lost node '%s' on connection '%s'..."), *NodeId.ToString(), *RemoteEndpoint.ToString());
+
+		NodeConnectionMapUpdates.Enqueue(FNodeConnectionMapUpdate(false, NodeId, TWeakPtr<FTcpMessageTransportConnection>(Connection)));
+		TransportHandler->ForgetTransportNode(NodeId);
 	}
 }
diff --git a/Engine/Plugins/Messaging/TcpMessaging/Source/TcpMessaging/Private/Transport/TcpMessageTransport.h b/Engine/Plugins/Messaging/TcpMessaging/Source/TcpMessaging/Private/Transport/TcpMessageTransport.h
--- a/Engine/Plugins/Messaging/TcpMessaging/Source/TcpMessaging/Private/Transport/TcpMessageTransport.h
+++ b/Engine/Plugins/Messaging/TcpMessaging/Source/TcpMessaging/Private/Transport/TcpMessageTransport.h
@@ -13,24 +13,26 @@ class FSocket;
 class FTcpListener;
 class FTcpMessageTransportConnection;
 
+
 /** Entry specfying addition or removal to/from the NodeConnectionMap */
 struct FNodeConnectionMapUpdate
 {
 	FNodeConnectionMapUpdate(bool bInNewNode, const FGuid& InNodeId, const TWeakPtr<FTcpMessageTransportConnection>& InConnection)
-	:	bNewNode(bInNewNode)
-	,	NodeId(InNodeId)
-	,	Connection(InConnection)
-	{}
+		: bNewNode(bInNewNode)
+		, NodeId(InNodeId)
+		, Connection(InConnection)
+	{ }
 
 	FNodeConnectionMapUpdate()
-	:	bNewNode(false)
-	{}
+		: bNewNode(false)
+	{ }
 
 	bool bNewNode;
 	FGuid NodeId;
 	TWeakPtr<FTcpMessageTransportConnection> Connection;
 };
 
+
 /**
  * Implements a message transport technology using an Tcp network connection.
  *
@@ -55,12 +57,10 @@ public:
 
 public:
 
-	// FTcpMessageTransport interface
-
 	void AddOutgoingConnection(const FIPv4Endpoint& Endpoint);
-
 	void RemoveOutgoingConnection(const FIPv4Endpoint& Endpoint);
 
+public:
 
 	//~ IMessageTransport interface
 
@@ -69,31 +69,18 @@ public:
 		return "TcpMessageTransport";
 	}
 
-	virtual FOnMessageReceived& OnMessageReceived() override
-	{
-		return MessageReceivedDelegate;
-	}
-
-	virtual FOnNodeDiscovered& OnNodeDiscovered() override
-	{
-		return NodeDiscoveredDelegate;
-	}
-
-	virtual FOnNodeLost& OnNodeLost() override
-	{
-		return NodeLostDelegate;
-	}
-
-	virtual bool StartTransport() override;
+	virtual bool StartTransport(IMessageTransportHandler& Handler) override;
 	virtual void StopTransport() override;
-	virtual bool TransportMessage(const IMessageContextRef& Context, const TArray<FGuid>& Recipients) override;
+	virtual bool TransportMessage(const TSharedRef<IMessageContext, ESPMode::ThreadSafe>& Context, const TArray<FGuid>& Recipients) override;
+
+public:
 
 	//~ FRunnable interface
 
+	virtual void Exit() override;
 	virtual bool Init() override;
 	virtual uint32 Run() override;
 	virtual void Stop() override;
-	virtual void Exit() override;
 
 private:
 
@@ -105,15 +92,6 @@ private:
 
 private:
 
-	/** Holds a delegate to be invoked when a message was received on the transport channel. */
-	FOnMessageReceived MessageReceivedDelegate;
-
-	/** Holds a delegate to be invoked when a network node was discovered. */
-	FOnNodeDiscovered NodeDiscoveredDelegate;
-
-	/** Holds a delegate to be invoked when a network node was lost. */
-	FOnNodeLost NodeLostDelegate;
-
 	/** Current settings */
 	FIPv4Endpoint ListenEndpoint;
 	TArray<FIPv4Endpoint> ConnectToEndpoints;
@@ -145,4 +123,7 @@ private:
 
 	/** Holds the thread object. */
 	FRunnableThread* Thread;
+
+	/** Message transport handler. */
+	IMessageTransportHandler* TransportHandler;
 };
diff --git a/Engine/Plugins/Messaging/TcpMessaging/Source/TcpMessaging/Private/Transport/TcpMessageTransportConnection.cpp b/Engine/Plugins/Messaging/TcpMessaging/Source/TcpMessaging/Private/Transport/TcpMessageTransportConnection.cpp
--- a/Engine/Plugins/Messaging/TcpMessaging/Source/TcpMessaging/Private/Transport/TcpMessageTransportConnection.cpp
+++ b/Engine/Plugins/Messaging/TcpMessaging/Source/TcpMessaging/Private/Transport/TcpMessageTransportConnection.cpp
@@ -380,7 +380,7 @@ bool FTcpMessageTransportConnection::BlockingSend(const uint8* Data, int32 Bytes
 	int32 TotalBytes = BytesToSend;
 	while (BytesToSend > 0)
 	{
-		while (!Socket->Wait(ESocketWaitConditions::WaitForWrite, FTimespan(0, 0, 1)))
+		while (!Socket->Wait(ESocketWaitConditions::WaitForWrite, FTimespan::FromSeconds(1.0)))
 		{
 			if (Socket->GetConnectionState() == SCS_ConnectionError)
 			{
diff --git a/Engine/Plugins/Messaging/TcpMessaging/Source/TcpMessaging/Private/Transport/TcpMessageTransportConnection.h b/Engine/Plugins/Messaging/TcpMessaging/Source/TcpMessaging/Private/Transport/TcpMessageTransportConnection.h
--- a/Engine/Plugins/Messaging/TcpMessaging/Source/TcpMessaging/Private/Transport/TcpMessageTransportConnection.h
+++ b/Engine/Plugins/Messaging/TcpMessaging/Source/TcpMessaging/Private/Transport/TcpMessageTransportConnection.h
@@ -222,7 +222,7 @@ private:
 	int32 ConnectionRetryDelay;
 
 	/** Message data we're currently in the process of receiving, if any */
-	FArrayReaderPtr RecvMessageData;
+	TSharedPtr<FArrayReader, ESPMode::ThreadSafe> RecvMessageData;
 
 	/** The number of bytes of incoming message data we're still waiting on before we have a complete message */
 	int32 RecvMessageDataRemaining;
diff --git a/Engine/Plugins/Messaging/TcpMessaging/Source/TcpMessaging/Private/Transport/TcpSerializeMessageTask.h b/Engine/Plugins/Messaging/TcpMessaging/Source/TcpMessaging/Private/Transport/TcpSerializeMessageTask.h
--- a/Engine/Plugins/Messaging/TcpMessaging/Source/TcpMessaging/Private/Transport/TcpSerializeMessageTask.h
+++ b/Engine/Plugins/Messaging/TcpMessaging/Source/TcpMessaging/Private/Transport/TcpSerializeMessageTask.h
@@ -22,7 +22,7 @@ public:
 	 * @param InMessageContext The context of the message to serialize.
 	 * @param InSerializedMessage Will hold the serialized message data.
 	 */
-	FTcpSerializeMessageTask(IMessageContextRef InMessageContext, FTcpSerializedMessageRef InSerializedMessage, const TArray<TSharedPtr<FTcpMessageTransportConnection>>& InRecipientConnections)
+	FTcpSerializeMessageTask(TSharedRef<IMessageContext, ESPMode::ThreadSafe> InMessageContext, FTcpSerializedMessageRef InSerializedMessage, const TArray<TSharedPtr<FTcpMessageTransportConnection>>& InRecipientConnections)
 		: MessageContext(InMessageContext)
 		, SerializedMessage(InSerializedMessage)
 		, RecipientConnections(InRecipientConnections)
@@ -62,7 +62,7 @@ public:
 private:
 
 	/** Holds the context of the message to serialize. */
-	IMessageContextRef MessageContext;
+	TSharedRef<IMessageContext, ESPMode::ThreadSafe> MessageContext;
 
 	/** Holds a reference to the serialized message data. */
 	FTcpSerializedMessageRef SerializedMessage;
diff --git a/Engine/Plugins/Messaging/TcpMessaging/Source/TcpMessaging/TcpMessaging.Build.cs b/Engine/Plugins/Messaging/TcpMessaging/Source/TcpMessaging/TcpMessaging.Build.cs
--- a/Engine/Plugins/Messaging/TcpMessaging/Source/TcpMessaging/TcpMessaging.Build.cs
+++ b/Engine/Plugins/Messaging/TcpMessaging/Source/TcpMessaging/TcpMessaging.Build.cs
@@ -6,18 +6,16 @@ namespace UnrealBuildTool.Rules
 	{
 		public TcpMessaging(ReadOnlyTargetRules Target) : base(Target)
 		{
-            DynamicallyLoadedModuleNames.AddRange(
-                new string[] {
-                    "Messaging",
-				}
-            );
+			DynamicallyLoadedModuleNames.AddRange(
+				new string[] {
+					"Messaging",
+				});
 
 			PublicDependencyModuleNames.AddRange(
 				new string[]
 				{
 					"Core",
-				}
-			); 
+				});
 
 			PrivateDependencyModuleNames.AddRange(
 				new string[] {
@@ -26,36 +24,32 @@ namespace UnrealBuildTool.Rules
 					"Networking",
 					"Serialization",
 					"Sockets",
-				}
-			);
+				});
 
 			PrivateIncludePathModuleNames.AddRange(
 				new string[] {
 					"Messaging",
-				}
-			);
+					"MessagingCommon",
+				});
 
 			PrivateIncludePaths.AddRange(
 				new string[] {
 					"TcpMessaging/Private",
 					"TcpMessaging/Private/Settings",
 					"TcpMessaging/Private/Transport",
-				}
-			);
+				});
 
 			if (Target.Type == TargetType.Editor)
 			{
 				DynamicallyLoadedModuleNames.AddRange(
 					new string[] {
-					"Settings",
-				}
-				);
+						"Settings",
+					});
 
 				PrivateIncludePathModuleNames.AddRange(
 					new string[] {
-					"Settings",
-				}
-				);
+						"Settings",
+					});
 			}
 		}
 	}
diff --git a/Engine/Plugins/Messaging/UdpMessaging/Source/UdpMessaging/Private/Shared/UdpMessageSegment.h b/Engine/Plugins/Messaging/UdpMessaging/Source/UdpMessaging/Private/Shared/UdpMessageSegment.h
--- a/Engine/Plugins/Messaging/UdpMessaging/Source/UdpMessaging/Private/Shared/UdpMessageSegment.h
+++ b/Engine/Plugins/Messaging/UdpMessaging/Source/UdpMessaging/Private/Shared/UdpMessageSegment.h
@@ -2,8 +2,11 @@
 
 #pragma once
 
-#include "CoreMinimal.h"
+#include "CoreTypes.h"
+#include "Containers/Array.h"
 #include "Misc/Guid.h"
+#include "Serialization/Archive.h"
+
 
 /**
  * Enumerates message segment types.
diff --git a/Engine/Plugins/Messaging/UdpMessaging/Source/UdpMessaging/Private/Shared/UdpMessagingSettings.h b/Engine/Plugins/Messaging/UdpMessaging/Source/UdpMessaging/Private/Shared/UdpMessagingSettings.h
--- a/Engine/Plugins/Messaging/UdpMessaging/Source/UdpMessaging/Private/Shared/UdpMessagingSettings.h
+++ b/Engine/Plugins/Messaging/UdpMessaging/Source/UdpMessaging/Private/Shared/UdpMessagingSettings.h
@@ -2,9 +2,12 @@
 
 #pragma once
 
-#include "CoreMinimal.h"
+#include "CoreTypes.h"
+#include "Containers/Array.h"
+#include "Containers/UnrealString.h"
 #include "UObject/ObjectMacros.h"
 #include "UObject/Object.h"
+
 #include "UdpMessagingSettings.generated.h"
 
 
diff --git a/Engine/Plugins/Messaging/UdpMessaging/Source/UdpMessaging/Private/Tests/UdpMessageTransportTest.cpp b/Engine/Plugins/Messaging/UdpMessaging/Source/UdpMessaging/Private/Tests/UdpMessageTransportTest.cpp
--- a/Engine/Plugins/Messaging/UdpMessaging/Source/UdpMessaging/Private/Tests/UdpMessageTransportTest.cpp
+++ b/Engine/Plugins/Messaging/UdpMessaging/Source/UdpMessaging/Private/Tests/UdpMessageTransportTest.cpp
@@ -6,6 +6,7 @@
 #include "Interfaces/IPv4/IPv4Address.h"
 #include "Interfaces/IPv4/IPv4Endpoint.h"
 #include "IMessageContext.h"
+#include "IMessageTransportHandler.h"
 #include "Transport/UdpMessageTransport.h"
 #include "Tests/UdpMessagingTestTypes.h"
 
@@ -14,6 +15,7 @@ IMPLEMENT_SIMPLE_AUTOMATION_TEST(FUdpMessageTransportTest, "System.Core.Messagin
 
 
 class FUdpMessageTransportTestState
+	: public IMessageTransportHandler
 {
 public:
 
@@ -22,11 +24,6 @@ public:
 		, Test(InTest)
 	{
 		Transport = MakeShareable(new FUdpMessageTransport(UnicastEndpoint, MulticastEndpoint, MulticastTimeToLive));
-		{
-			Transport->OnMessageReceived().BindRaw(this, &FUdpMessageTransportTestState::HandleTransportMessageReceived);
-			Transport->OnNodeDiscovered().BindRaw(this, &FUdpMessageTransportTestState::HandleTransportNodeDiscovered);
-			Transport->OnNodeLost().BindRaw(this, &FUdpMessageTransportTestState::HandleTransportNodeLost);
-		}
 	}
 
 public:
@@ -53,7 +50,7 @@ public:
 
 	bool Start()
 	{
-		return Transport->StartTransport();
+		return Transport->StartTransport(*this);
 	}
 
 	void Stop()
@@ -61,23 +58,25 @@ public:
 		Transport->StopTransport();
 	}
 
-private:
+public:
 
-	void HandleTransportMessageReceived(const TSharedRef<IMessageContext, ESPMode::ThreadSafe>& MessageContext, const FGuid& NodeId)
-	{
-		FPlatformAtomics::InterlockedIncrement(&NumReceivedMessages);
-	}
+	//~ IMessageTransportHandler interface
 
-	void HandleTransportNodeDiscovered(const FGuid& NodeId)
+	virtual void DiscoverTransportNode(const FGuid& NodeId) override
 	{
 		DiscoveredNodes.Add(NodeId);
 	}
 
-	void HandleTransportNodeLost(const FGuid& NodeId)
+	virtual void ForgetTransportNode(const FGuid& NodeId) override
 	{
 		LostNodes.Add(NodeId);
 	}
 
+	virtual void ReceiveTransportMessage(const TSharedRef<IMessageContext, ESPMode::ThreadSafe>& Context, const FGuid& NodeId) override
+	{
+		FPlatformAtomics::InterlockedIncrement(&NumReceivedMessages);
+	}
+
 private:
 
 	TArray<FGuid> DiscoveredNodes;
diff --git a/Engine/Plugins/Messaging/UdpMessaging/Source/UdpMessaging/Private/Tests/UdpMessagingTestTypes.h b/Engine/Plugins/Messaging/UdpMessaging/Source/UdpMessaging/Private/Tests/UdpMessagingTestTypes.h
--- a/Engine/Plugins/Messaging/UdpMessaging/Source/UdpMessaging/Private/Tests/UdpMessagingTestTypes.h
+++ b/Engine/Plugins/Messaging/UdpMessaging/Source/UdpMessaging/Private/Tests/UdpMessagingTestTypes.h
@@ -80,7 +80,7 @@ private:
 	TSharedPtr<IMessageAttachment, ESPMode::ThreadSafe> Attachment;
 	FDateTime Expiration;
 	void* Message;
-	IMessageContextPtr OriginalContext;
+	TSharedPtr<IMessageContext, ESPMode::ThreadSafe> OriginalContext;
 	TArray<FMessageAddress> Recipients;
 	EMessageScope Scope;
 	FMessageAddress Sender;
