Merging CL#3747555 from Dev-Networking

Do not treat sceNetEpollWait() returning 0 as an error
Fix tcp messaging connection retry socket not having send/receive buffer sizes set
Change some verbose logs to warnings

diff --git a/Engine/Plugins/Messaging/TcpMessaging/Source/TcpMessaging/Private/TcpMessagingPrivate.h b/Engine/Plugins/Messaging/TcpMessaging/Source/TcpMessaging/Private/TcpMessagingPrivate.h
--- a/Engine/Plugins/Messaging/TcpMessaging/Source/TcpMessaging/Private/TcpMessagingPrivate.h
+++ b/Engine/Plugins/Messaging/TcpMessaging/Source/TcpMessaging/Private/TcpMessagingPrivate.h
@@ -16,6 +16,9 @@ DECLARE_LOG_CATEGORY_EXTERN(LogTcpMessaging, Log, All);
 /** Defines the maximum number of recipients a message can have. */
 #define TCP_MESSAGING_MAX_RECIPIENTS 1024
 
+/** Defines the desired size of socket send buffers (in bytes). */
+#define TCP_MESSAGING_SEND_BUFFER_SIZE 2 * 1024 * 1024
+
 /** Defines the desired size of socket receive buffers (in bytes). */
 #define TCP_MESSAGING_RECEIVE_BUFFER_SIZE 2 * 1024 * 1024
 
diff --git a/Engine/Plugins/Messaging/TcpMessaging/Source/TcpMessaging/Private/Transport/TcpMessageTransportConnection.cpp b/Engine/Plugins/Messaging/TcpMessaging/Source/TcpMessaging/Private/Transport/TcpMessageTransportConnection.cpp
--- a/Engine/Plugins/Messaging/TcpMessaging/Source/TcpMessaging/Private/Transport/TcpMessageTransportConnection.cpp
+++ b/Engine/Plugins/Messaging/TcpMessaging/Source/TcpMessaging/Private/Transport/TcpMessageTransportConnection.cpp
@@ -71,8 +71,8 @@ FTcpMessageTransportConnection::FTcpMessageTransportConnection(FSocket* InSocket
 	, RecvMessageDataRemaining(0)
 {
 	int32 NewSize = 0;
-	Socket->SetReceiveBufferSize(2*1024*1024, NewSize);
-	Socket->SetSendBufferSize(2*1024*1024, NewSize);
+	Socket->SetReceiveBufferSize(TCP_MESSAGING_RECEIVE_BUFFER_SIZE, NewSize);
+	Socket->SetSendBufferSize(TCP_MESSAGING_SEND_BUFFER_SIZE, NewSize);
 }
 
 FTcpMessageTransportConnection::~FTcpMessageTransportConnection()
@@ -180,10 +180,13 @@ uint32 FTcpMessageTransportConnection::Run()
 				    ISocketSubsystem::Get(PLATFORM_SOCKETSUBSYSTEM)->DestroySocket(Socket);
 				    Socket = nullptr;
     
-				    UE_LOG(LogTcpMessaging, Verbose, TEXT("Connection to '%s' failed, retrying..."), *RemoteEndpoint.ToString());
+				    UE_LOG(LogTcpMessaging, Warning, TEXT("Connection to '%s' failed, retrying..."), *RemoteEndpoint.ToString());
 				    FPlatformProcess::Sleep(ConnectionRetryDelay);
     
-				    Socket = FTcpSocketBuilder(TEXT("FTcpMessageTransport.RemoteConnection"));
+				    Socket = FTcpSocketBuilder(TEXT("FTcpMessageTransport.RemoteConnection"))
+					.WithSendBufferSize(TCP_MESSAGING_SEND_BUFFER_SIZE)
+					.WithReceiveBufferSize(TCP_MESSAGING_RECEIVE_BUFFER_SIZE);
+
 				    if (Socket && Socket->Connect(RemoteEndpoint.ToInternetAddr().Get()))
 				    {
 					    bSentHeader = false;
@@ -312,7 +315,7 @@ bool FTcpMessageTransportConnection::ReceiveMessages()
 		uint8 Dummy;
 		if (!Socket->Recv(&Dummy, 1, BytesRead, ESocketReceiveFlags::Peek))
 		{
-			UE_LOG(LogTcpMessaging, Verbose, TEXT("Dummy read failed with code %d"), (int32)SocketSubsystem->GetLastErrorCode());
+			UE_LOG(LogTcpMessaging, Warning, TEXT("Dummy read failed with code %d"), (int32)SocketSubsystem->GetLastErrorCode());
 			return false;
 		}
 	}
@@ -332,7 +335,7 @@ bool FTcpMessageTransportConnection::ReceiveMessages()
 			int32 BytesRead = 0;
 			if (!Socket->Recv(HeaderData.GetData(), sizeof(FTcpMessageHeader), BytesRead))
 			{
-				UE_LOG(LogTcpMessaging, Verbose, TEXT("Header read failed with code %d"), (int32)SocketSubsystem->GetLastErrorCode());
+				UE_LOG(LogTcpMessaging, Warning, TEXT("Header read failed with code %d"), (int32)SocketSubsystem->GetLastErrorCode());
 				return false;
 			}
 
@@ -344,7 +347,7 @@ bool FTcpMessageTransportConnection::ReceiveMessages()
 
 			if (!MessageHeader.IsValid())
 			{
-				UE_LOG(LogTcpMessaging, Verbose, TEXT("Header read failed with invalid header"));
+				UE_LOG(LogTcpMessaging, Warning, TEXT("Header read failed with invalid header"));
 				return false;
 			}
 			else
@@ -388,7 +391,7 @@ bool FTcpMessageTransportConnection::ReceiveMessages()
 			BytesRead = 0;
 			if (!Socket->Recv(MessagesizeData.GetData(), sizeof(uint32), BytesRead))
 			{
-				UE_LOG(LogTcpMessaging, Verbose, TEXT("In progress read failed with code %d"), (int32)SocketSubsystem->GetLastErrorCode());
+				UE_LOG(LogTcpMessaging, Warning, TEXT("In progress read failed with code %d"), (int32)SocketSubsystem->GetLastErrorCode());
 				return false;
 			}
 
@@ -406,7 +409,7 @@ bool FTcpMessageTransportConnection::ReceiveMessages()
 		BytesRead = 0;
 		if (!Socket->Recv(RecvMessageData->GetData() + RecvMessageData->Num() - RecvMessageDataRemaining, RecvMessageDataRemaining, BytesRead))
 		{
-			UE_LOG(LogTcpMessaging, Verbose, TEXT("Read failed with code %d"), (int32)SocketSubsystem->GetLastErrorCode());
+			UE_LOG(LogTcpMessaging, Warning, TEXT("Read failed with code %d"), (int32)SocketSubsystem->GetLastErrorCode());
 			return false;
 		}
 
@@ -481,7 +484,7 @@ bool FTcpMessageTransportConnection::SendHeader()
 
 	if (!BlockingSend(HeaderData.GetData(), sizeof(FTcpMessageHeader)))
 	{
-		UE_LOG(LogTcpMessaging, Verbose, TEXT("Header write failed with code %d"), (int32)SocketSubsystem->GetLastErrorCode());
+		UE_LOG(LogTcpMessaging, Warning, TEXT("Header write failed with code %d"), (int32)SocketSubsystem->GetLastErrorCode());
 		return false;
 	}
 
