Revert log verbosity change from cl#3778235

diff --git a/Engine/Plugins/Messaging/TcpMessaging/Source/TcpMessaging/Private/Transport/TcpMessageTransportConnection.cpp b/Engine/Plugins/Messaging/TcpMessaging/Source/TcpMessaging/Private/Transport/TcpMessageTransportConnection.cpp
--- a/Engine/Plugins/Messaging/TcpMessaging/Source/TcpMessaging/Private/Transport/TcpMessageTransportConnection.cpp
+++ b/Engine/Plugins/Messaging/TcpMessaging/Source/TcpMessaging/Private/Transport/TcpMessageTransportConnection.cpp
@@ -180,7 +180,7 @@ uint32 FTcpMessageTransportConnection::Run()
 				    ISocketSubsystem::Get(PLATFORM_SOCKETSUBSYSTEM)->DestroySocket(Socket);
 				    Socket = nullptr;
     
-				    UE_LOG(LogTcpMessaging, Warning, TEXT("Connection to '%s' failed, retrying..."), *RemoteEndpoint.ToString());
+				    UE_LOG(LogTcpMessaging, Verbose, TEXT("Connection to '%s' failed, retrying..."), *RemoteEndpoint.ToString());
 				    FPlatformProcess::Sleep(ConnectionRetryDelay);
     
 				    Socket = FTcpSocketBuilder(TEXT("FTcpMessageTransport.RemoteConnection"))
@@ -315,7 +315,7 @@ bool FTcpMessageTransportConnection::ReceiveMessages()
 		uint8 Dummy;
 		if (!Socket->Recv(&Dummy, 1, BytesRead, ESocketReceiveFlags::Peek))
 		{
-			UE_LOG(LogTcpMessaging, Warning, TEXT("Dummy read failed with code %d"), (int32)SocketSubsystem->GetLastErrorCode());
+			UE_LOG(LogTcpMessaging, Verbose, TEXT("Dummy read failed with code %d"), (int32)SocketSubsystem->GetLastErrorCode());
 			return false;
 		}
 	}
@@ -335,7 +335,7 @@ bool FTcpMessageTransportConnection::ReceiveMessages()
 			int32 BytesRead = 0;
 			if (!Socket->Recv(HeaderData.GetData(), sizeof(FTcpMessageHeader), BytesRead))
 			{
-				UE_LOG(LogTcpMessaging, Warning, TEXT("Header read failed with code %d"), (int32)SocketSubsystem->GetLastErrorCode());
+				UE_LOG(LogTcpMessaging, Verbose, TEXT("Header read failed with code %d"), (int32)SocketSubsystem->GetLastErrorCode());
 				return false;
 			}
 
@@ -347,7 +347,7 @@ bool FTcpMessageTransportConnection::ReceiveMessages()
 
 			if (!MessageHeader.IsValid())
 			{
-				UE_LOG(LogTcpMessaging, Warning, TEXT("Header read failed with invalid header"));
+				UE_LOG(LogTcpMessaging, Verbose, TEXT("Header read failed with invalid header"));
 				return false;
 			}
 			else
@@ -391,7 +391,7 @@ bool FTcpMessageTransportConnection::ReceiveMessages()
 			BytesRead = 0;
 			if (!Socket->Recv(MessagesizeData.GetData(), sizeof(uint32), BytesRead))
 			{
-				UE_LOG(LogTcpMessaging, Warning, TEXT("In progress read failed with code %d"), (int32)SocketSubsystem->GetLastErrorCode());
+				UE_LOG(LogTcpMessaging, Verbose, TEXT("In progress read failed with code %d"), (int32)SocketSubsystem->GetLastErrorCode());
 				return false;
 			}
 
@@ -409,7 +409,7 @@ bool FTcpMessageTransportConnection::ReceiveMessages()
 		BytesRead = 0;
 		if (!Socket->Recv(RecvMessageData->GetData() + RecvMessageData->Num() - RecvMessageDataRemaining, RecvMessageDataRemaining, BytesRead))
 		{
-			UE_LOG(LogTcpMessaging, Warning, TEXT("Read failed with code %d"), (int32)SocketSubsystem->GetLastErrorCode());
+			UE_LOG(LogTcpMessaging, Verbose, TEXT("Read failed with code %d"), (int32)SocketSubsystem->GetLastErrorCode());
 			return false;
 		}
 
@@ -484,7 +484,7 @@ bool FTcpMessageTransportConnection::SendHeader()
 
 	if (!BlockingSend(HeaderData.GetData(), sizeof(FTcpMessageHeader)))
 	{
-		UE_LOG(LogTcpMessaging, Warning, TEXT("Header write failed with code %d"), (int32)SocketSubsystem->GetLastErrorCode());
+		UE_LOG(LogTcpMessaging, Verbose, TEXT("Header write failed with code %d"), (int32)SocketSubsystem->GetLastErrorCode());
 		return false;
 	}
 
