Fix TcpMessaging to close connections that send invalid / corrupt messages.

diff --git a/Engine/Plugins/Messaging/TcpMessaging/Source/TcpMessaging/Private/Transport/TcpMessageTransportConnection.cpp b/Engine/Plugins/Messaging/TcpMessaging/Source/TcpMessaging/Private/Transport/TcpMessageTransportConnection.cpp
--- a/Engine/Plugins/Messaging/TcpMessaging/Source/TcpMessaging/Private/Transport/TcpMessageTransportConnection.cpp
+++ b/Engine/Plugins/Messaging/TcpMessaging/Source/TcpMessaging/Private/Transport/TcpMessageTransportConnection.cpp
@@ -311,13 +311,19 @@ bool FTcpMessageTransportConnection::ReceiveMessages()
 	ISocketSubsystem* SocketSubsystem = ISocketSubsystem::Get(PLATFORM_SOCKETSUBSYSTEM);
 	uint32 PendingDataSize = 0;
 	
+	auto GetReadableErrorCode = [SocketSubsystem]() -> FString
+	{
+		ESocketErrors LastError = SocketSubsystem->GetLastErrorCode();
+		return FString::Printf(TEXT("%s (%d)"), SocketSubsystem->GetSocketError(LastError), (int32)LastError);
+	};
+
 	// check if the socket has closed
 	{
 		int32 BytesRead;
 		uint8 Dummy;
 		if (!Socket->Recv(&Dummy, 1, BytesRead, ESocketReceiveFlags::Peek))
 		{
-			UE_LOG(LogTcpMessaging, Verbose, TEXT("Dummy read failed with code %d"), (int32)SocketSubsystem->GetLastErrorCode());
+			UE_LOG(LogTcpMessaging, Verbose, TEXT("Dummy read failed with code %s"), *GetReadableErrorCode());
 			return false;
 		}
 	}
@@ -337,7 +343,7 @@ bool FTcpMessageTransportConnection::ReceiveMessages()
 			int32 BytesRead = 0;
 			if (!Socket->Recv(HeaderData.GetData(), sizeof(FTcpMessageHeader), BytesRead))
 			{
-				UE_LOG(LogTcpMessaging, Verbose, TEXT("Header read failed with code %d"), (int32)SocketSubsystem->GetLastErrorCode());
+				UE_LOG(LogTcpMessaging, Verbose, TEXT("Header read failed with code %s"), *GetReadableErrorCode());
 				return false;
 			}
 
@@ -393,7 +399,7 @@ bool FTcpMessageTransportConnection::ReceiveMessages()
 			BytesRead = 0;
 			if (!Socket->Recv(MessagesizeData.GetData(), sizeof(uint32), BytesRead))
 			{
-				UE_LOG(LogTcpMessaging, Verbose, TEXT("In progress read failed with code %d"), (int32)SocketSubsystem->GetLastErrorCode());
+				UE_LOG(LogTcpMessaging, Verbose, TEXT("In progress read failed with code %s"), *GetReadableErrorCode());
 				return false;
 			}
 
@@ -403,15 +409,20 @@ bool FTcpMessageTransportConnection::ReceiveMessages()
 			// Setup variables to receive the message
 			MessagesizeData << RecvMessageDataRemaining;
 
+			if (RecvMessageDataRemaining <= 0)
+			{
+				UE_LOG(LogTcpMessaging, Error, TEXT("Read failed due to invalid Message Size: %d"), RecvMessageDataRemaining);
+				return false;
+			}
+
 			RecvMessageData = MakeShareable(new FArrayReader(true));
 			RecvMessageData->SetNumUninitialized(RecvMessageDataRemaining);
-			check(RecvMessageDataRemaining > 0);
 		}
 
 		BytesRead = 0;
 		if (!Socket->Recv(RecvMessageData->GetData() + RecvMessageData->Num() - RecvMessageDataRemaining, RecvMessageDataRemaining, BytesRead))
 		{
-			UE_LOG(LogTcpMessaging, Verbose, TEXT("Read failed with code %d"), (int32)SocketSubsystem->GetLastErrorCode());
+			UE_LOG(LogTcpMessaging, Verbose, TEXT("Read failed with code %s"), *GetReadableErrorCode());
 			return false;
 		}
 
