Fix client not connecting to TCP Messaging server on high latency network.
* Set Socket Wait For Write to 5 seconds when sending header, otherwise the client may close the socket and retry too fast for the connection to be established.

#rnx

diff --git a/Engine/Plugins/Messaging/TcpMessaging/Source/TcpMessaging/Private/Transport/TcpMessageTransportConnection.cpp b/Engine/Plugins/Messaging/TcpMessaging/Source/TcpMessaging/Private/Transport/TcpMessageTransportConnection.cpp
--- a/Engine/Plugins/Messaging/TcpMessaging/Source/TcpMessaging/Private/Transport/TcpMessageTransportConnection.cpp
+++ b/Engine/Plugins/Messaging/TcpMessaging/Source/TcpMessaging/Private/Transport/TcpMessageTransportConnection.cpp
@@ -486,7 +486,7 @@ bool FTcpMessageTransportConnection::SendHeader()
 	ISocketSubsystem* SocketSubsystem = ISocketSubsystem::Get(PLATFORM_SOCKETSUBSYSTEM);
 
 	// See if we're writable
-	if (!Socket->Wait(ESocketWaitConditions::WaitForWrite, FTimespan::Zero()))
+	if (!Socket->Wait(ESocketWaitConditions::WaitForWrite, FTimespan::FromSeconds(5)))
 	{
 		return true;
 	}
