PR #8055: Tcp multiple connections fix (Contributed by [REDACTED_PERSON])

diff --git a/Engine/Plugins/Messaging/TcpMessaging/Source/TcpMessaging/Private/TcpMessagingModule.cpp b/Engine/Plugins/Messaging/TcpMessaging/Source/TcpMessaging/Private/TcpMessagingModule.cpp
--- a/Engine/Plugins/Messaging/TcpMessaging/Source/TcpMessaging/Private/TcpMessagingModule.cpp
+++ b/Engine/Plugins/Messaging/TcpMessaging/Source/TcpMessaging/Private/TcpMessagingModule.cpp
@@ -375,7 +375,7 @@ void UTcpMessagingSettings::GetConnectToEndpoints(TArray<FString>& Endpoints) co
 	}
 
 	FString ConnectString;
-	if (FParse::Value(FCommandLine::Get(), TEXT("TcpMessagingConnect="), ConnectString))
+	if (FParse::Value(FCommandLine::Get(), TEXT("TcpMessagingConnect="), ConnectString, /* bShouldStopOnSeparator */ false))
 	{
 		ConnectString.ParseIntoArray(Endpoints, TEXT(","));
 	}
