Improve TcpMessageTransportConnection to support chunked sending when buffer size cannot be accomodated.

#rnx

diff --git a/Engine/Plugins/Messaging/TcpMessaging/Source/TcpMessaging/Private/Transport/TcpMessageTransportConnection.cpp b/Engine/Plugins/Messaging/TcpMessaging/Source/TcpMessaging/Private/Transport/TcpMessageTransportConnection.cpp
--- a/Engine/Plugins/Messaging/TcpMessaging/Source/TcpMessaging/Private/Transport/TcpMessageTransportConnection.cpp
+++ b/Engine/Plugins/Messaging/TcpMessaging/Source/TcpMessaging/Private/Transport/TcpMessageTransportConnection.cpp
@@ -113,7 +113,6 @@ bool FTcpMessageTransportConnection::Receive(TSharedPtr<FTcpDeserializedMessage,
 	return false;
 }
 
-
 bool FTcpMessageTransportConnection::Send(FTcpSerializedMessagePtr Message)
 {
 	FScopeLock SendLock(&SendCriticalSection);
@@ -137,7 +136,7 @@ bool FTcpMessageTransportConnection::Send(FTcpSerializedMessagePtr Message)
 		TotalBytesSent += sizeof(uint32);
 
 		// send the payload
-		if (!BlockingSend(Payload.GetData(), Payload.Num()))
+		if (!ChunkedBlockingSend(Payload))
 		{
 			UE_LOG(LogTcpMessaging, Verbose, TEXT("Payload write failed with code %d"), (int32)ISocketSubsystem::Get(PLATFORM_SOCKETSUBSYSTEM)->GetLastErrorCode());
 			return false;
@@ -292,7 +291,7 @@ FTimespan FTcpMessageTransportConnection::GetUptime() const
 	{
 		return (FDateTime::UtcNow() - OpenedTime);
 	}
-		
+
 	return (ClosedTime - OpenedTime);
 }
 
@@ -310,7 +309,7 @@ bool FTcpMessageTransportConnection::ReceiveMessages()
 {
 	ISocketSubsystem* SocketSubsystem = ISocketSubsystem::Get(PLATFORM_SOCKETSUBSYSTEM);
 	uint32 PendingDataSize = 0;
-	
+
 	auto GetReadableErrorCode = [SocketSubsystem]() -> FString
 	{
 		ESocketErrors LastError = SocketSubsystem->GetLastErrorCode();
@@ -327,13 +326,13 @@ bool FTcpMessageTransportConnection::ReceiveMessages()
 			return false;
 		}
 	}
-	
+
 	// Block waiting for some data
 	if (!Socket->Wait(ESocketWaitConditions::WaitForRead, FTimespan::FromSeconds(1.0)))
 	{
 		return (Socket->GetConnectionState() != SCS_ConnectionError);
 	}
-	
+
 	if (!bReceivedHeader)
 	{
 		if (Socket->HasPendingData(PendingDataSize) && PendingDataSize >= sizeof(FTcpMessageHeader))
@@ -364,11 +363,11 @@ bool FTcpMessageTransportConnection::ReceiveMessages()
 				RemoteProtocolVersion = MessageHeader.GetVersion();
 				bReceivedHeader = true;
 				OpenedTime = FDateTime::UtcNow();
-	            {
-		            FScopeLock SendLock(&SendCriticalSection);
-		            ConnectionState = STATE_Connected;
+				{
+					FScopeLock SendLock(&SendCriticalSection);
+					ConnectionState = STATE_Connected;
 				}
-	            ConnectionStateChangedDelegate.ExecuteIfBound();
+				ConnectionStateChangedDelegate.ExecuteIfBound();
 			}
 		}
 		else
@@ -450,6 +449,55 @@ bool FTcpMessageTransportConnection::ReceiveMessages()
 	}
 }
 
+bool FTcpMessageTransportConnection::ChunkedBlockingSend(const TArray<uint8>& Payload)
+{
+	// send the payload
+	if (!BlockingSend(Payload.GetData(), Payload.Num()))
+	{
+		ESocketErrors LastError = ISocketSubsystem::Get(PLATFORM_SOCKETSUBSYSTEM)->GetLastErrorCode();
+		UE_LOG(LogTcpMessaging, Verbose, TEXT("Payload write failed with code %d"), (int32)LastError);
+
+		if (LastError == SE_ENOBUFS)
+		{
+			UE_LOG(LogTcpMessaging, Verbose, TEXT("Could not send %d bytes in one blocking send. Trying really small chunks."), Payload.Num());
+
+			const int32 Max64kb = 64*1024;
+			const uint32 MaxRetries = 16;
+			uint32 NumRetries = 0;
+
+			FPlatformProcess::Sleep(1);
+			const uint8* Buf = Payload.GetData();
+			int32 BufSz = Payload.Num();
+			while (BufSz > 0)
+			{
+				int SendMax = FMath::Min(BufSz, Max64kb);
+				if (BlockingSend(Buf, SendMax))
+				{
+					Buf += SendMax;
+					BufSz -= SendMax;
+				}
+				else if (NumRetries < MaxRetries)
+				{
+					FPlatformProcess::Sleep(1);
+					NumRetries++;
+					UE_LOG(LogTcpMessaging, Verbose, TEXT("%d retries sending small chunks"), NumRetries);
+				}
+				else
+				{
+					// The target could have disconnected in the middle of the transmission.
+					return false;
+				}
+			}
+		}
+		else
+		{
+			return false;
+		}
+	}
+
+	return true;
+}
+
 bool FTcpMessageTransportConnection::BlockingSend(const uint8* Data, int32 BytesToSend)
 {
 	int32 TotalBytes = BytesToSend;
diff --git a/Engine/Plugins/Messaging/TcpMessaging/Source/TcpMessaging/Private/Transport/TcpMessageTransportConnection.h b/Engine/Plugins/Messaging/TcpMessaging/Source/TcpMessaging/Private/Transport/TcpMessageTransportConnection.h
--- a/Engine/Plugins/Messaging/TcpMessaging/Source/TcpMessaging/Private/Transport/TcpMessageTransportConnection.h
+++ b/Engine/Plugins/Messaging/TcpMessaging/Source/TcpMessaging/Private/Transport/TcpMessageTransportConnection.h
@@ -157,6 +157,10 @@ private:
 	virtual void Exit() override;
 
 private:
+
+	/** Split the send data into chunks until all data is sent. */
+	bool ChunkedBlockingSend(const TArray<uint8>& Payload);
+
 	/** Try to send data, but if all data is not sent in one go, block on send until data is sent or an error occurs */
 	bool BlockingSend(const uint8* Data, int32 BytesToSend);
 
