Improve TCPMessaging retry on connection failure mechanic

diff --git a/Engine/Plugins/Messaging/TcpMessaging/Source/TcpMessaging/Private/Settings/TcpMessagingSettings.h b/Engine/Plugins/Messaging/TcpMessaging/Source/TcpMessaging/Private/Settings/TcpMessagingSettings.h
--- a/Engine/Plugins/Messaging/TcpMessaging/Source/TcpMessaging/Private/Settings/TcpMessagingSettings.h
+++ b/Engine/Plugins/Messaging/TcpMessaging/Source/TcpMessaging/Private/Settings/TcpMessagingSettings.h
@@ -25,9 +25,12 @@ public:
 	/** Returns array of connect endpoints */
 	void GetConnectToEndpoints(TArray<FString>& Endpoints) const;
 
-	/** Returns connection delay */
+	/** Returns reconnection delay */
 	int32 GetConnectionRetryDelay() const;
 
+	/** Returns reconnection period */
+	int32 GetConnectionRetryPeriod() const;
+
 	/** Returns whether to stop the transport service when the application deactivates */
 	bool ShouldStopServiceWhenAppDeactivates() const;
 
@@ -61,6 +64,13 @@ private:
 	UPROPERTY(config, EditAnywhere, Category = Transport)
 	int32 ConnectionRetryDelay;
 
+	/**
+	 * Period time during which attempts to re-establish outgoing connections that become disconnected or fail to connect
+	 * 0 means it will be retried only once
+	 */
+	UPROPERTY(config, EditAnywhere, Category = Transport)
+	int32 ConnectionRetryPeriod;
+
 	/** Whether to stop the transport service when the application deactivates, and restart it when the application is reactivated */
 	UPROPERTY(config)
 	bool bStopServiceWhenAppDeactivates = true;
diff --git a/Engine/Plugins/Messaging/TcpMessaging/Source/TcpMessaging/Private/TcpMessagingModule.cpp b/Engine/Plugins/Messaging/TcpMessaging/Source/TcpMessaging/Private/TcpMessagingModule.cpp
--- a/Engine/Plugins/Messaging/TcpMessaging/Source/TcpMessaging/Private/TcpMessagingModule.cpp
+++ b/Engine/Plugins/Messaging/TcpMessaging/Source/TcpMessaging/Private/TcpMessagingModule.cpp
@@ -242,7 +242,7 @@ protected:
 
 		UE_LOG(LogTcpMessaging, Log, TEXT("%s"), *Status);
 
-		TSharedRef<FTcpMessageTransport, ESPMode::ThreadSafe> Transport = MakeShareable(new FTcpMessageTransport(ListenEndpoint, ConnectToEndpoints, Settings->GetConnectionRetryDelay()));
+		TSharedRef<FTcpMessageTransport, ESPMode::ThreadSafe> Transport = MakeShareable(new FTcpMessageTransport(ListenEndpoint, ConnectToEndpoints, Settings->GetConnectionRetryDelay(), Settings->GetConnectionRetryPeriod()));
 		
 		// Safe weak pointer for adding/removing connections
 		MessageTransportPtr = Transport;
@@ -386,6 +386,11 @@ int32 UTcpMessagingSettings::GetConnectionRetryDelay() const
 	return ConnectionRetryDelay;
 }
 
+int32 UTcpMessagingSettings::GetConnectionRetryPeriod() const
+{
+	return ConnectionRetryPeriod;
+}
+
 bool UTcpMessagingSettings::ShouldStopServiceWhenAppDeactivates() const
 {
 	return bStopServiceWhenAppDeactivates;
diff --git a/Engine/Plugins/Messaging/TcpMessaging/Source/TcpMessaging/Private/Transport/TcpMessageTransport.cpp b/Engine/Plugins/Messaging/TcpMessaging/Source/TcpMessaging/Private/Transport/TcpMessageTransport.cpp
--- a/Engine/Plugins/Messaging/TcpMessaging/Source/TcpMessaging/Private/Transport/TcpMessageTransport.cpp
+++ b/Engine/Plugins/Messaging/TcpMessaging/Source/TcpMessaging/Private/Transport/TcpMessageTransport.cpp
@@ -17,10 +17,11 @@
 /* FTcpMessageTransport structors
  *****************************************************************************/
 
-FTcpMessageTransport::FTcpMessageTransport(const FIPv4Endpoint& InListenEndpoint, const TArray<FIPv4Endpoint>& InConnectToEndpoints, int32 InConnectionRetryDelay)
+FTcpMessageTransport::FTcpMessageTransport(const FIPv4Endpoint& InListenEndpoint, const TArray<FIPv4Endpoint>& InConnectToEndpoints, int32 InConnectionRetryDelay, int32 InConnectionRetryPeriod)
 	: ListenEndpoint(InListenEndpoint)
 	, ConnectToEndpoints(InConnectToEndpoints)
 	, ConnectionRetryDelay(InConnectionRetryDelay)
+	, ConnectionRetryPeriod(InConnectionRetryPeriod)
 	, bStopping(false)
 	, SocketSubsystem(ISocketSubsystem::Get(PLATFORM_SOCKETSUBSYSTEM))
 	, Listener(nullptr)
@@ -81,7 +82,7 @@ void FTcpMessageTransport::AddOutgoingConnection(const FIPv4Endpoint& Endpoint)
 	}
 	else
 	{
-		PendingConnections.Enqueue(MakeShareable(new FTcpMessageTransportConnection(Socket, Endpoint, ConnectionRetryDelay)));
+		PendingConnections.Enqueue(MakeShareable(new FTcpMessageTransportConnection(Socket, Endpoint, ConnectionRetryDelay, ConnectionRetryPeriod)));
 	}
 }
 
diff --git a/Engine/Plugins/Messaging/TcpMessaging/Source/TcpMessaging/Private/Transport/TcpMessageTransport.h b/Engine/Plugins/Messaging/TcpMessaging/Source/TcpMessaging/Private/Transport/TcpMessageTransport.h
--- a/Engine/Plugins/Messaging/TcpMessaging/Source/TcpMessaging/Private/Transport/TcpMessageTransport.h
+++ b/Engine/Plugins/Messaging/TcpMessaging/Source/TcpMessaging/Private/Transport/TcpMessageTransport.h
@@ -50,7 +50,7 @@ public:
 	 * Creates and initializes a new instance.
 	 *
 	 */
-	FTcpMessageTransport(const FIPv4Endpoint& InListenEndpoint, const TArray<FIPv4Endpoint>& InConnectToEndpoints, int32 InConnectionRetryDelay);
+	FTcpMessageTransport(const FIPv4Endpoint& InListenEndpoint, const TArray<FIPv4Endpoint>& InConnectToEndpoints, int32 InConnectionRetryDelay, int32 InConnectionRetryPeriod = 0);
 
 	/** Virtual destructor. */
 	virtual ~FTcpMessageTransport();
@@ -96,6 +96,7 @@ private:
 	FIPv4Endpoint ListenEndpoint;
 	TArray<FIPv4Endpoint> ConnectToEndpoints;
 	int32 ConnectionRetryDelay;
+	int32 ConnectionRetryPeriod;
 
 	/** For the thread */
 	bool bStopping;
diff --git a/Engine/Plugins/Messaging/TcpMessaging/Source/TcpMessaging/Private/Transport/TcpMessageTransportConnection.cpp b/Engine/Plugins/Messaging/TcpMessaging/Source/TcpMessaging/Private/Transport/TcpMessageTransportConnection.cpp
--- a/Engine/Plugins/Messaging/TcpMessaging/Source/TcpMessaging/Private/Transport/TcpMessageTransportConnection.cpp
+++ b/Engine/Plugins/Messaging/TcpMessaging/Source/TcpMessaging/Private/Transport/TcpMessageTransportConnection.cpp
@@ -54,7 +54,7 @@ struct FTcpMessageHeader
 /* FTcpMessageTransportConnection structors
  *****************************************************************************/
 
-FTcpMessageTransportConnection::FTcpMessageTransportConnection(FSocket* InSocket, const FIPv4Endpoint& InRemoteEndpoint, int32 InConnectionRetryDelay)
+FTcpMessageTransportConnection::FTcpMessageTransportConnection(FSocket* InSocket, const FIPv4Endpoint& InRemoteEndpoint, int32 InConnectionRetryDelay, int32 InConnectionRetryPeriod)
 	: ConnectionState(STATE_Connecting)
 	, OpenedTime(FDateTime::UtcNow())
 	, RemoteEndpoint(InRemoteEndpoint)
@@ -68,6 +68,7 @@ FTcpMessageTransportConnection::FTcpMessageTransportConnection(FSocket* InSocket
 	, TotalBytesSent(0)
 	, bRun(false)
 	, ConnectionRetryDelay(InConnectionRetryDelay)
+	, ConnectionRetryPeriod(InConnectionRetryPeriod)
 	, RecvMessageDataRemaining(0)
 {
 	int32 NewSize = 0;
@@ -172,6 +173,7 @@ uint32 FTcpMessageTransportConnection::Run()
 			if (ConnectionRetryDelay > 0)
 			{
 				bool bReconnectPending = false;
+				float TimeSpentRetrying = 0;
 
 				{
 				    // Wait for any sending before we close the socket
@@ -179,38 +181,47 @@ uint32 FTcpMessageTransportConnection::Run()
     
 				    Socket->Close();
 				    ISocketSubsystem::Get(PLATFORM_SOCKETSUBSYSTEM)->DestroySocket(Socket);
-				    Socket = nullptr;
-    
-				    UE_LOG(LogTcpMessaging, Verbose, TEXT("Connection to '%s' failed, retrying..."), *RemoteEndpoint.ToString());
-				    FPlatformProcess::Sleep(ConnectionRetryDelay);
-    
-				    Socket = FTcpSocketBuilder(TEXT("FTcpMessageTransport.RemoteConnection"))
-						.WithSendBufferSize(TCP_MESSAGING_SEND_BUFFER_SIZE)
-						.WithReceiveBufferSize(TCP_MESSAGING_RECEIVE_BUFFER_SIZE);
-
-				    if (Socket && Socket->Connect(RemoteEndpoint.ToInternetAddr().Get()))
-				    {
-					    bSentHeader = false;
-					    bReceivedHeader = false;
-						ConnectionState = STATE_DisconnectReconnectPending;
-					    RemoteNodeId.Invalidate();
-						bReconnectPending = true;
-				    }
-				    else
-				    {
-					    if (Socket)
-					    {
-						    ISocketSubsystem::Get(PLATFORM_SOCKETSUBSYSTEM)->DestroySocket(Socket);
-						    Socket = nullptr;
-					    }
-					    bRun = false;
-				    }
+					Socket = nullptr;
+
+					while (TimeSpentRetrying <= ConnectionRetryPeriod && Socket == nullptr)
+					{
+						UE_LOG(LogTcpMessaging, Log, TEXT("Connection to '%s' failed, retrying..."), *RemoteEndpoint.ToString());
+						FPlatformProcess::Sleep(ConnectionRetryDelay);
+						TimeSpentRetrying += ConnectionRetryDelay;
+
+						Socket = FTcpSocketBuilder(TEXT("FTcpMessageTransport.RemoteConnection"))
+							.WithSendBufferSize(TCP_MESSAGING_SEND_BUFFER_SIZE)
+							.WithReceiveBufferSize(TCP_MESSAGING_RECEIVE_BUFFER_SIZE);
+
+						if (Socket && Socket->Connect(RemoteEndpoint.ToInternetAddr().Get()))
+						{
+							bSentHeader = false;
+							bReceivedHeader = false;
+							ConnectionState = STATE_DisconnectReconnectPending;
+							RemoteNodeId.Invalidate();
+							bReconnectPending = true;
+						}
+						else
+						{
+							if (Socket)
+							{
+								ISocketSubsystem::Get(PLATFORM_SOCKETSUBSYSTEM)->DestroySocket(Socket);
+								Socket = nullptr;
+							}
+						}
+					}
+
+					bRun = Socket != nullptr;
 				}
 
 				if (bReconnectPending)
 				{
 					ConnectionStateChangedDelegate.ExecuteIfBound();
 				}
+				else
+				{
+					UE_LOG(LogTcpMessaging, Error, TEXT("Reconnection to '%s' failed."), *RemoteEndpoint.ToString());
+				}
 			}
 			else
 			{
diff --git a/Engine/Plugins/Messaging/TcpMessaging/Source/TcpMessaging/Private/Transport/TcpMessageTransportConnection.h b/Engine/Plugins/Messaging/TcpMessaging/Source/TcpMessaging/Private/Transport/TcpMessageTransportConnection.h
--- a/Engine/Plugins/Messaging/TcpMessaging/Source/TcpMessaging/Private/Transport/TcpMessageTransportConnection.h
+++ b/Engine/Plugins/Messaging/TcpMessaging/Source/TcpMessaging/Private/Transport/TcpMessageTransportConnection.h
@@ -39,8 +39,9 @@ public:
 	 * @param InSocket The socket to use for this connection.
 	 * @param InRemoteAddress The IP endpoint of the remote client.
 	 * @param int32 InConnectionRetryDelay Delay before retrying connection, 0 disables
+	 * @param int32 InConnectionRetryPeriod period during which retrying connection, 0 means to retry only once
 	 */
-	FTcpMessageTransportConnection(FSocket* InSocket, const FIPv4Endpoint& InRemoteEndpoint, int32 InConnectionRetryDelay);
+	FTcpMessageTransportConnection(FSocket* InSocket, const FIPv4Endpoint& InRemoteEndpoint, int32 InConnectionRetryDelay, int32 InConnectionRetryPeriod = 0);
 
 	/** Virtual destructor. */
 	virtual ~FTcpMessageTransportConnection();
@@ -215,6 +216,9 @@ private:
 	/** Delay before re-establishing connection if it drops, 0 disables */
 	int32 ConnectionRetryDelay;
 
+	/** Period during which re-establishing connection if it drops, 0 means it will retry only once */
+	int32 ConnectionRetryPeriod;
+
 	/** Message data we're currently in the process of receiving, if any */
 	TSharedPtr<FArrayReader, ESPMode::ThreadSafe> RecvMessageData;
 
