Change log level to verbose when attempting reconnection on TCPMessaging

diff --git a/Engine/Plugins/Messaging/TcpMessaging/Source/TcpMessaging/Private/Transport/TcpMessageTransportConnection.cpp b/Engine/Plugins/Messaging/TcpMessaging/Source/TcpMessaging/Private/Transport/TcpMessageTransportConnection.cpp
--- a/Engine/Plugins/Messaging/TcpMessaging/Source/TcpMessaging/Private/Transport/TcpMessageTransportConnection.cpp
+++ b/Engine/Plugins/Messaging/TcpMessaging/Source/TcpMessaging/Private/Transport/TcpMessageTransportConnection.cpp
@@ -185,7 +185,7 @@ uint32 FTcpMessageTransportConnection::Run()
 
 					while (TimeSpentRetrying <= ConnectionRetryPeriod && Socket == nullptr)
 					{
-						UE_LOG(LogTcpMessaging, Log, TEXT("Connection to '%s' failed, retrying..."), *RemoteEndpoint.ToString());
+						UE_LOG(LogTcpMessaging, Verbose, TEXT("Connection to '%s' failed, retrying..."), *RemoteEndpoint.ToString());
 						FPlatformProcess::Sleep(ConnectionRetryDelay);
 						TimeSpentRetrying += ConnectionRetryDelay;
 
