Fix infinite loop with TCP Messaging when closing during reconnection attempt

diff --git a/Engine/Plugins/Messaging/TcpMessaging/Source/TcpMessaging/Private/Transport/TcpMessageTransportConnection.cpp b/Engine/Plugins/Messaging/TcpMessaging/Source/TcpMessaging/Private/Transport/TcpMessageTransportConnection.cpp
--- a/Engine/Plugins/Messaging/TcpMessaging/Source/TcpMessaging/Private/Transport/TcpMessageTransportConnection.cpp
+++ b/Engine/Plugins/Messaging/TcpMessaging/Source/TcpMessaging/Private/Transport/TcpMessageTransportConnection.cpp
@@ -170,11 +170,10 @@ uint32 FTcpMessageTransportConnection::Run()
 		if ((!SendHeader() || !ReceiveMessages() || Socket->GetConnectionState() == SCS_ConnectionError) && bRun)
 		{
 			// Disconnected. Reconnect if requested.
-			if (ConnectionRetryDelay > 0)
+			const float Delay = ConnectionRetryDelay;
+			if (Delay > 0)
 			{
 				bool bReconnectPending = false;
-				float TimeSpentRetrying = 0;
-
 				{
 				    // Wait for any sending before we close the socket
 				    FScopeLock SendLock(&SendCriticalSection);
@@ -183,11 +182,10 @@ uint32 FTcpMessageTransportConnection::Run()
 				    ISocketSubsystem::Get(PLATFORM_SOCKETSUBSYSTEM)->DestroySocket(Socket);
 					Socket = nullptr;
 
-					while (TimeSpentRetrying <= ConnectionRetryPeriod && Socket == nullptr)
+					for (float TimeSpentRetrying = 0; bRun && Socket == nullptr && TimeSpentRetrying <= ConnectionRetryPeriod; TimeSpentRetrying += Delay)
 					{
 						UE_LOG(LogTcpMessaging, Verbose, TEXT("Connection to '%s' failed, retrying..."), *RemoteEndpoint.ToString());
 						FPlatformProcess::Sleep(ConnectionRetryDelay);
-						TimeSpentRetrying += ConnectionRetryDelay;
 
 						Socket = FTcpSocketBuilder(TEXT("FTcpMessageTransport.RemoteConnection"))
 							.WithSendBufferSize(TCP_MESSAGING_SEND_BUFFER_SIZE)
@@ -207,11 +205,10 @@ uint32 FTcpMessageTransportConnection::Run()
 							{
 								ISocketSubsystem::Get(PLATFORM_SOCKETSUBSYSTEM)->DestroySocket(Socket);
 								Socket = nullptr;
+								bRun = false;
 							}
 						}
 					}
-
-					bRun = Socket != nullptr;
 				}
 
 				if (bReconnectPending)
