Make sure GC is not in progress when deserializing message

diff --git a/Engine/Plugins/Messaging/TcpMessaging/Source/TcpMessaging/Private/Transport/TcpDeserializedMessage.cpp b/Engine/Plugins/Messaging/TcpMessaging/Source/TcpMessaging/Private/Transport/TcpDeserializedMessage.cpp
--- a/Engine/Plugins/Messaging/TcpMessaging/Source/TcpMessaging/Private/Transport/TcpDeserializedMessage.cpp
+++ b/Engine/Plugins/Messaging/TcpMessaging/Source/TcpMessaging/Private/Transport/TcpDeserializedMessage.cpp
@@ -49,12 +49,17 @@ bool FTcpDeserializedMessage::Deserialize(const TSharedPtr<FArrayReader, ESPMode
 		FTopLevelAssetPath MessageType;
 		MessageReader << MessageType;
 
-		// @todo gmp: cache message types for faster lookup
+		while (IsGarbageCollectingAndLockingUObjectHashTables())
+		{
+			// Wait for garbage collector
+			FPlatformProcess::Sleep(0.1f);
+		}
 		TypeInfoPtr = FindObjectSafe<UScriptStruct>(MessageType);
 
 		TypeInfo = TypeInfoPtr;
 		if (!TypeInfo.IsValid(false, true))
 		{
+			UE_LOG(LogTcpMessaging, Verbose, TEXT("Lookup for Message Type '%s' failed during Deserialization."), *MessageType.ToString());
 			return false;
 		}
 	}
diff --git a/Engine/Plugins/Messaging/TcpMessaging/Source/TcpMessaging/Private/Transport/TcpMessageTransportConnection.cpp b/Engine/Plugins/Messaging/TcpMessaging/Source/TcpMessaging/Private/Transport/TcpMessageTransportConnection.cpp
--- a/Engine/Plugins/Messaging/TcpMessaging/Source/TcpMessaging/Private/Transport/TcpMessageTransportConnection.cpp
+++ b/Engine/Plugins/Messaging/TcpMessaging/Source/TcpMessaging/Private/Transport/TcpMessageTransportConnection.cpp
@@ -440,11 +440,10 @@ bool FTcpMessageTransportConnection::ReceiveMessages()
 			RecvMessageDataRemaining -= BytesRead;
 			if (RecvMessageDataRemaining == 0)
 			{
-				// @todo gmp: move message deserialization into an async task
-				FTcpDeserializedMessage* DeserializedMessage = new FTcpDeserializedMessage(nullptr);
+				TSharedPtr<FTcpDeserializedMessage, ESPMode::ThreadSafe> DeserializedMessage = MakeShareable(new FTcpDeserializedMessage(nullptr));
 				if (DeserializedMessage->Deserialize(RecvMessageData))
 				{
-					Inbox.Enqueue(MakeShareable(DeserializedMessage));
+					Inbox.Enqueue(DeserializedMessage);
 				}
 				RecvMessageData.Reset();
 			}
