Fix TCPMessaging supported version

diff --git a/Engine/Plugins/Messaging/TcpMessaging/Source/TcpMessaging/Private/TcpMessagingPrivate.h b/Engine/Plugins/Messaging/TcpMessaging/Source/TcpMessaging/Private/TcpMessagingPrivate.h
--- a/Engine/Plugins/Messaging/TcpMessaging/Source/TcpMessaging/Private/TcpMessagingPrivate.h
+++ b/Engine/Plugins/Messaging/TcpMessaging/Source/TcpMessaging/Private/TcpMessagingPrivate.h
@@ -39,6 +39,6 @@ namespace ETcpMessagingVersion
 		VersionPlusOne,
 		LatestVersion = VersionPlusOne - 1,
 		// bump this when break need to break compatibility.
-		OldestSupportedVersion = ChangedMessageLengthToInt32
+		OldestSupportedVersion = ChangedMessageContext
 	};
 }
diff --git a/Engine/Plugins/Messaging/TcpMessaging/Source/TcpMessaging/Private/Transport/TcpMessageTransportConnection.cpp b/Engine/Plugins/Messaging/TcpMessaging/Source/TcpMessaging/Private/Transport/TcpMessageTransportConnection.cpp
--- a/Engine/Plugins/Messaging/TcpMessaging/Source/TcpMessaging/Private/Transport/TcpMessageTransportConnection.cpp
+++ b/Engine/Plugins/Messaging/TcpMessaging/Source/TcpMessaging/Private/Transport/TcpMessageTransportConnection.cpp
@@ -31,7 +31,7 @@ struct FTcpMessageHeader
 	{
 		return
 			MagicNumber == TCP_MESSAGING_TRANSPORT_PROTOCOL_MAGIC &&
-			Version == ETcpMessagingVersion::OldestSupportedVersion &&
+			Version >= ETcpMessagingVersion::OldestSupportedVersion &&
 			NodeId.IsValid();
 	}
 
