Revert change to the oldest supported version to match what was shipped in 5.1

diff --git a/Engine/Plugins/Messaging/TcpMessaging/Source/TcpMessaging/Private/TcpMessagingPrivate.h b/Engine/Plugins/Messaging/TcpMessaging/Source/TcpMessaging/Private/TcpMessagingPrivate.h
--- a/Engine/Plugins/Messaging/TcpMessaging/Source/TcpMessaging/Private/TcpMessagingPrivate.h
+++ b/Engine/Plugins/Messaging/TcpMessaging/Source/TcpMessaging/Private/TcpMessagingPrivate.h
@@ -39,6 +39,6 @@ namespace ETcpMessagingVersion
 		VersionPlusOne,
 		LatestVersion = VersionPlusOne - 1,
 		// bump this when break need to break compatibility.
-		OldestSupportedVersion = ChangedMessageContext
+		OldestSupportedVersion = ChangedMessageLengthToInt32
 	};
 }
