UnrealPak references TcpMessaging module

#rnx

diff --git a/Engine/Plugins/Messaging/TcpMessaging/TcpMessaging.uplugin b/Engine/Plugins/Messaging/TcpMessaging/TcpMessaging.uplugin
--- a/Engine/Plugins/Messaging/TcpMessaging/TcpMessaging.uplugin
+++ b/Engine/Plugins/Messaging/TcpMessaging/TcpMessaging.uplugin
@@ -14,7 +14,7 @@
 	"CanContainContent" : false,
 	"IsBetaVersion" : false,
 	"Installed" : false,
-	"SupportedPrograms": [ "UnrealFrontend" , "UnrealInsights" ],
+	"SupportedPrograms": [ "UnrealFrontend" , "UnrealInsights", "UnrealPak" ],
 	"Modules" :
 	[
 		{
@@ -24,7 +24,7 @@
 			"PlatformDenyList" :
 			[
 			],
-			"ProgramAllowList": [ "UnrealFrontend", "UnrealInsights" ]
+			"ProgramAllowList": [ "UnrealFrontend", "UnrealInsights", "UnrealPak" ]
 	    }
 	]
 }
